Editor: Documentation: Updated the documentation code to allow proper fallback behavior for cultures missing documentation.

diff --git a/Engine/Source/Editor/Documentation/Private/Documentation.cpp b/Engine/Source/Editor/Documentation/Private/Documentation.cpp
--- a/Engine/Source/Editor/Documentation/Private/Documentation.cpp
+++ b/Engine/Source/Editor/Documentation/Private/Documentation.cpp
@@ -30,6 +30,11 @@ bool FDocumentation::OpenHome() const
 	return Open( TEXT("%ROOT%") );
 }
 
+bool FDocumentation::OpenHome(const TSharedRef<FCulture>& Culture) const
+{
+	return Open(TEXT("%ROOT%"), Culture);
+}
+
 bool FDocumentation::OpenAPIHome() const
 {
 	FString APIPath = FPaths::Combine(*FPaths::EngineDir(), TEXT("Documentation/CHM/API.chm"));
@@ -50,29 +55,63 @@ bool FDocumentation::Open( const FString& Link ) const
 {
 	FString DocumentationUrl;
 
-	if( !FParse::Param( FCommandLine::Get(), TEXT("testdocs") ) )
+	if (!FParse::Param(FCommandLine::Get(), TEXT("testdocs")))
 	{
-		const FString OnDiskPath = FDocumentationLink::ToFilePath(Link);
+		FString OnDiskPath = FDocumentationLink::ToFilePath(Link);
 		if (IFileManager::Get().FileSize(*OnDiskPath) != INDEX_NONE)
 		{
 			DocumentationUrl = FDocumentationLink::ToFileUrl(Link);
 		}
 	}
 
-	if( DocumentationUrl.IsEmpty() )
+	if (DocumentationUrl.IsEmpty())
+	{
+		// When opening a doc website we always request the most ideal culture for our documentation.
+		// The DNS will redirect us if necessary.
+		DocumentationUrl = FDocumentationLink::ToUrl(Link);
+	}
+
+	if (!DocumentationUrl.IsEmpty())
+	{
+		FPlatformProcess::LaunchURL(*DocumentationUrl, NULL, NULL);
+	}
+
+	if (!DocumentationUrl.IsEmpty() && FEngineAnalytics::IsAvailable())
+	{
+		FEngineAnalytics::GetProvider().RecordEvent(TEXT("Editor.Usage.Documentation"), TEXT("OpenedPage"), Link);
+	}
+
+	return !DocumentationUrl.IsEmpty();
+}
+
+bool FDocumentation::Open(const FString& Link, const TSharedRef<FCulture>& Culture) const
+{
+	FString DocumentationUrl;
+
+	if (!FParse::Param(FCommandLine::Get(), TEXT("testdocs")))
+	{
+		FString OnDiskPath = FDocumentationLink::ToFilePath(Link, Culture);
+		if (IFileManager::Get().FileSize(*OnDiskPath) != INDEX_NONE)
+		{
+			DocumentationUrl = FDocumentationLink::ToFileUrl(Link, Culture);
+		}
+	}
+
+	if (DocumentationUrl.IsEmpty())
 	{
-		DocumentationUrl = FDocumentationLink::ToUrl( Link );
+		DocumentationUrl = FDocumentationLink::ToUrl(Link, Culture);
 	}
 
-	if( !DocumentationUrl.IsEmpty() )
+	if (!DocumentationUrl.IsEmpty())
 	{
-		FPlatformProcess::LaunchURL( *DocumentationUrl, NULL, NULL );
+		FPlatformProcess::LaunchURL(*DocumentationUrl, NULL, NULL);
 	}
 
-	if( !DocumentationUrl.IsEmpty() && FEngineAnalytics::IsAvailable() )
+	if (!DocumentationUrl.IsEmpty() && FEngineAnalytics::IsAvailable())
 	{
-		FEngineAnalytics::GetProvider().RecordEvent(TEXT("Editor.Usage.Documentation"), TEXT("OpenedPage"), Link );
+		FEngineAnalytics::GetProvider().RecordEvent(TEXT("Editor.Usage.Documentation"), TEXT("OpenedPage"), Link);
 	}
+
 	return !DocumentationUrl.IsEmpty();
 }
 
@@ -114,8 +153,20 @@ bool FDocumentation::PageExists(const FString& Link) const
 		return true;
 	}
 
-	const FString SourcePath = FDocumentationLink::ToSourcePath(Link, true);
-	return !SourcePath.IsEmpty();
+	const FString SourcePath = FDocumentationLink::ToSourcePath(Link);
+	return FPaths::FileExists(SourcePath);
+}
+
+bool FDocumentation::PageExists(const FString& Link, const TSharedRef<FCulture>& Culture) const
+{
+	const TWeakPtr< IDocumentationPage >* ExistingPagePtr = LoadedPages.Find(Link);
+	if (ExistingPagePtr != NULL)
+	{
+		return true;
+	}
+
+	const FString SourcePath = FDocumentationLink::ToSourcePath(Link, Culture);
+	return FPaths::FileExists(SourcePath);
 }
 
 TSharedRef< class SToolTip > FDocumentation::CreateToolTip( const TAttribute<FText>& Text, const TSharedPtr<SWidget>& OverrideContent, const FString& Link, const FString& ExcerptName ) const
diff --git a/Engine/Source/Editor/Documentation/Private/Documentation.h b/Engine/Source/Editor/Documentation/Private/Documentation.h
--- a/Engine/Source/Editor/Documentation/Private/Documentation.h
+++ b/Engine/Source/Editor/Documentation/Private/Documentation.h
@@ -15,16 +15,22 @@ public:
 
 	virtual bool OpenHome() const OVERRIDE;
 
+	virtual bool OpenHome(const TSharedRef<FCulture>& Culture) const OVERRIDE;
+
 	virtual bool OpenAPIHome() const OVERRIDE;
 
 	virtual bool Open( const FString& Link ) const OVERRIDE;
 
+	virtual bool Open(const FString& Link, const TSharedRef<FCulture>& Culture) const OVERRIDE;
+
 	virtual TSharedRef< SWidget > CreateAnchor( const FString& Link, const FString& PreviewLink = FString(), const FString& PreviewExcerptName = FString() ) const OVERRIDE;
 
 	virtual TSharedRef< IDocumentationPage > GetPage( const FString& Link, const TSharedPtr< FParserConfiguration >& Config, const FDocumentationStyle& Style = FDocumentationStyle() ) OVERRIDE;
 
 	virtual bool PageExists(const FString& Link) const OVERRIDE;
 
+	virtual bool PageExists(const FString& Link, const TSharedRef<FCulture>& Culture) const OVERRIDE;
+
 	virtual TSharedRef< class SToolTip > CreateToolTip( const TAttribute<FText>& Text, const TSharedPtr<SWidget>& OverrideContent, const FString& Link, const FString& ExcerptName ) const OVERRIDE;
 
 private:
diff --git a/Engine/Source/Editor/Documentation/Private/DocumentationLink.h b/Engine/Source/Editor/Documentation/Private/DocumentationLink.h
--- a/Engine/Source/Editor/Documentation/Private/DocumentationLink.h
+++ b/Engine/Source/Editor/Documentation/Private/DocumentationLink.h
@@ -19,45 +19,116 @@ public:
 	}
 
 	static FString GetHomeUrl()
+	{
+		return GetHomeUrl(FInternationalization::GetCurrentCulture());
+	}
+
+	static FString GetHomeUrl(const TSharedRef<FCulture>& Culture)
 	{
 		FString Url;
 		FUnrealEdMisc::Get().GetURL( TEXT("UDNURL"), Url, true );
 
-		Url.ReplaceInline( TEXT( "/INT/" ), *FString::Printf( TEXT( "/%s/" ), *(FInternationalization::GetCurrentCulture()->GetUnrealLegacyThreeLetterISOLanguageName()) ) );
+		Url.ReplaceInline(TEXT("/INT/"), *FString::Printf(TEXT("/%s/"), *(Culture->GetUnrealLegacyThreeLetterISOLanguageName())));
 		return Url;
 	}
 
-	static FString ToUrl( const FString& Link )
+	static FString ToUrl(const FString& Link)
+	{
+		return ToUrl(Link, FInternationalization::GetCurrentCulture());
+	}
+
+	static FString ToUrl(const FString& Link, const TSharedRef<FCulture>& Culture)
 	{
 		FString Path;
 		FString Anchor;
 		SplitLink( Link, Path, Anchor );
 
-		const FString PartialPath = FString::Printf( TEXT( "%s%s/index.html" ), *(FInternationalization::GetCurrentCulture()->GetUnrealLegacyThreeLetterISOLanguageName()), *Path );
+		const FString PartialPath = FString::Printf(TEXT("%s%s/index.html"), *(Culture->GetUnrealLegacyThreeLetterISOLanguageName()), *Path);
 		
 		return GetUrlRoot() + PartialPath + Anchor;
 	}
 
 	static FString ToFilePath( const FString& Link )
+	{
+		FString FilePath = ToFilePath(Link, FInternationalization::GetCurrentCulture());
+
+		if (!FPaths::FileExists(FilePath))
+		{
+			const TSharedPtr<FCulture> FallbackCulture = FInternationalization::GetCulture(TEXT("en"));
+			if (FallbackCulture.IsValid())
+			{
+				const FString FallbackFilePath = ToFilePath(Link, FallbackCulture.ToSharedRef());
+				if (FPaths::FileExists(FallbackFilePath))
+				{
+					FilePath = FallbackFilePath;
+				}
+			}
+		}
+
+		return FilePath;
+	}
+
+	static FString ToFilePath(const FString& Link, const TSharedRef<FCulture>& Culture)
 	{
 		FString Path;
 		FString Anchor;
-		SplitLink( Link, Path, Anchor );
+		SplitLink(Link, Path, Anchor);
 
-		const FString PartialPath = FString::Printf( TEXT( "%s%s/index.html" ), *(FInternationalization::GetCurrentCulture()->GetUnrealLegacyThreeLetterISOLanguageName()), *Path );
-		return FString::Printf( TEXT( "%sDocumentation/HTML/%s" ), *FPaths::ConvertRelativePathToFull( FPaths::EngineDir() ), *PartialPath );
+		const FString PartialPath = FString::Printf(TEXT("%s%s/index.html"), *(Culture->GetUnrealLegacyThreeLetterISOLanguageName()), *Path);
+		return FString::Printf(TEXT("%sDocumentation/HTML/%s"), *FPaths::ConvertRelativePathToFull(FPaths::EngineDir()), *PartialPath);
 	}
 
 	static FString ToFileUrl( const FString& Link )
+	{
+		TSharedRef<FCulture> Culture = FInternationalization::GetCurrentCulture();
+		FString FilePath = ToFilePath(Link, Culture);
+
+		if (!FPaths::FileExists(FilePath))
+		{
+			const TSharedPtr<FCulture> FallbackCulture = FInternationalization::GetCulture(TEXT("en"));
+			if (FallbackCulture.IsValid())
+			{
+				const FString FallbackFilePath = ToFilePath(Link, FallbackCulture.ToSharedRef());
+				if (FPaths::FileExists(FallbackFilePath))
+				{
+					Culture = FallbackCulture.ToSharedRef();
+				}
+			}
+		}
+
+		return ToFileUrl(Link, Culture);
+	}
+
+	static FString ToFileUrl(const FString& Link, const TSharedRef<FCulture>& Culture)
 	{
 		FString Path;
 		FString Anchor;
-		SplitLink( Link, Path, Anchor );
+		SplitLink(Link, Path, Anchor);
 
-		return FString::Printf( TEXT( "file:///%s%s" ), *ToFilePath( Link ), *Anchor );
+		return FString::Printf(TEXT("file:///%s%s"), *ToFilePath(Link, Culture), *Anchor);
 	}
 
-	static FString ToSourcePath( const FString& Link, bool MustExist = true )
+	static FString ToSourcePath(const FString& Link)
+	{
+		FString SourcePath = ToSourcePath(Link, FInternationalization::GetCurrentCulture());
+
+		if (!FPaths::FileExists(SourcePath))
+		{
+			const TSharedPtr<FCulture> FallbackCulture = FInternationalization::GetCulture(TEXT("en"));
+			if (FallbackCulture.IsValid())
+			{
+				const FString FallbackSourcePath = ToSourcePath(Link, FallbackCulture.ToSharedRef());
+				if (FPaths::FileExists(FallbackSourcePath))
+				{
+					SourcePath = FallbackSourcePath;
+				}
+			}
+		}
+
+		return SourcePath;
+	}
+
+	static FString ToSourcePath(const FString& Link, const TSharedRef<FCulture>& Culture)
 	{
 		FString Path;
 		FString Anchor;
@@ -65,25 +136,19 @@ public:
 
 		const FString FullDirectoryPath = FPaths::EngineDir() + TEXT( "Documentation/Source" ) + Path + "/";
 
-		if ( MustExist )
-		{
-			const FString WildCard = FString::Printf( TEXT( "%s*.%s.udn" ), *FullDirectoryPath, *( FInternationalization::GetCurrentCulture()->GetUnrealLegacyThreeLetterISOLanguageName() ) );
-
-			TArray<FString> Filenames;
-			IFileManager::Get().FindFiles( Filenames, *WildCard, true, false );
+		const FString WildCard = FString::Printf(TEXT("%s*.%s.udn"), *FullDirectoryPath, *(Culture->GetUnrealLegacyThreeLetterISOLanguageName()));
 
-			if ( Filenames.Num() > 0 )
-			{
-				return FullDirectoryPath + Filenames[0];
-			}
+		TArray<FString> Filenames;
+		IFileManager::Get().FindFiles(Filenames, *WildCard, true, false);
 
-			return FString();
-		}
-		else
+		if (Filenames.Num() > 0)
 		{
-			FString Category = FPaths::GetBaseFilename( Link );
-			return FString::Printf( TEXT( "%s%s.%s.udn" ), *FullDirectoryPath, *Category, *( FInternationalization::GetCurrentCulture()->GetUnrealLegacyThreeLetterISOLanguageName() ) );
+			return FullDirectoryPath + Filenames[0];
 		}
+
+		// Since the source file doesn't exist already make up a valid name for a new one
+		FString Category = FPaths::GetBaseFilename(Link);
+		return FString::Printf(TEXT("%s%s.%s.udn"), *FullDirectoryPath, *Category, *(Culture->GetUnrealLegacyThreeLetterISOLanguageName()));
 	}
 
 	static void SplitLink( const FString& Link, /*OUT*/ FString& Path, /*OUT*/ FString& Anchor )
diff --git a/Engine/Source/Editor/Documentation/Private/SDocumentationToolTip.cpp b/Engine/Source/Editor/Documentation/Private/SDocumentationToolTip.cpp
--- a/Engine/Source/Editor/Documentation/Private/SDocumentationToolTip.cpp
+++ b/Engine/Source/Editor/Documentation/Private/SDocumentationToolTip.cpp
@@ -90,14 +90,14 @@ void SDocumentationToolTip::ConstructSimpleTipContent()
 			}
 
 			VerticalBox->AddSlot()
-				.AutoHeight()
-				.HAlign( HAlign_Right )
-				[
-					SNew( STextBlock )
-					.ColorAndOpacity( FSlateColor::UseSubduedForeground() )
-					.Text( DocumentationLink + OptionalExcerptName )
-					.TextStyle( &StyleInfo )
-				];
+			.AutoHeight()
+			.HAlign( HAlign_Right )
+			[
+				SNew( STextBlock )
+				.ColorAndOpacity( FSlateColor::UseSubduedForeground() )
+				.Text( DocumentationLink + OptionalExcerptName )
+				.TextStyle( &StyleInfo )
+			];
 		}
 
 		if ( !DocumentationPage.IsValid() )
@@ -118,33 +118,33 @@ void SDocumentationToolTip::ConstructSimpleTipContent()
 #endif
 
 			VerticalBox->AddSlot()
-				.AutoHeight()
-				.HAlign( HAlign_Right )
-				[
-					SNew( STextBlock )
-					.ColorAndOpacity( FSlateColor::UseSubduedForeground() )
-					.Text( FText::Format( NSLOCTEXT( "SToolTip", "AdvancedToolTipMessage", "hold {0} for more" ), KeyboardShortcut) )
-					.TextStyle( &StyleInfo )
-				];
+			.AutoHeight()
+			.HAlign( HAlign_Right )
+			[
+				SNew( STextBlock )
+				.ColorAndOpacity( FSlateColor::UseSubduedForeground() )
+				.Text( FText::Format( NSLOCTEXT( "SToolTip", "AdvancedToolTipMessage", "hold {0} for more" ), KeyboardShortcut) )
+				.TextStyle( &StyleInfo )
+			];
 		}
 		else
 		{
 			if ( GEditor->EditorUserSettings->bDisplayDocumentationLink && FSlateApplication::Get().SupportsSourceAccess() )
 			{
-				FString DocPath = FDocumentationLink::ToSourcePath( DocumentationLink, true );
-				if ( DocPath.IsEmpty() )
+				FString DocPath = FDocumentationLink::ToSourcePath( DocumentationLink, FInternationalization::GetCurrentCulture() );
+				if ( !FPaths::FileExists(DocPath) )
 				{
-					DocPath = FPaths::ConvertRelativePathToFull( FDocumentationLink::ToSourcePath( DocumentationLink, false ) );
+					DocPath = FPaths::ConvertRelativePathToFull(DocPath);
 				}
 
 				VerticalBox->AddSlot()
-					.AutoHeight()
-					.HAlign( HAlign_Right )
-					[
-						SNew( SHyperlink )
-						.Text( NSLOCTEXT( "SToolTip", "EditDocumentationMessage_Create", "create" ) )
-						.OnNavigate( this, &SDocumentationToolTip::CreateExcerpt, DocPath, ExcerptName )
-					];
+				.AutoHeight()
+				.HAlign( HAlign_Right )
+				[
+					SNew( SHyperlink )
+					.Text( NSLOCTEXT( "SToolTip", "EditDocumentationMessage_Create", "create" ) )
+					.OnNavigate( this, &SDocumentationToolTip::CreateExcerpt, DocPath, ExcerptName )
+				];
 			}
 		}
 	}
@@ -283,7 +283,7 @@ void SDocumentationToolTip::ConstructFullTipContent()
 				[
 					SNew( SHyperlink )
 						.Text( NSLOCTEXT( "SToolTip", "EditDocumentationMessage_Edit", "edit" ) )
-						.OnNavigate_Static( &Local::EditSource, FPaths::ConvertRelativePathToFull( FDocumentationLink::ToSourcePath( DocumentationLink ) ) + TEXT("|") + FString::FromInt( Excerpts[ ExcerptIndex ].LineNumber ) )
+						.OnNavigate_Static(&Local::EditSource, FPaths::ConvertRelativePathToFull(FDocumentationLink::ToSourcePath(DocumentationLink, FInternationalization::GetCurrentCulture())) + TEXT("|") + FString::FromInt(Excerpts[ExcerptIndex].LineNumber))
 				];
 			}
 		}
diff --git a/Engine/Source/Editor/Documentation/Private/UDNParser.cpp b/Engine/Source/Editor/Documentation/Private/UDNParser.cpp
--- a/Engine/Source/Editor/Documentation/Private/UDNParser.cpp
+++ b/Engine/Source/Editor/Documentation/Private/UDNParser.cpp
@@ -216,8 +216,8 @@ bool FUDNParser::LoadLink( const FString& Link, TArray<FString>& ContentLines )
 	FMessageLog UDNParserLog(UDNParseErrorLog);
 
 	const FString SourcePath = FDocumentationLink::ToSourcePath( Link );
-
-	if ( SourcePath.IsEmpty() )
+	
+	if (!FPaths::FileExists(SourcePath))
 	{
 		return false;
 	}
diff --git a/Engine/Source/Editor/Documentation/Public/IDocumentation.h b/Engine/Source/Editor/Documentation/Public/IDocumentation.h
--- a/Engine/Source/Editor/Documentation/Public/IDocumentation.h
+++ b/Engine/Source/Editor/Documentation/Public/IDocumentation.h
@@ -139,15 +139,21 @@ public:
 
 	virtual bool OpenHome() const = 0;
 
+	virtual bool OpenHome(const TSharedRef<FCulture>& Culture) const = 0;
+
 	virtual bool OpenAPIHome() const = 0;
 
 	virtual bool Open( const FString& Link ) const = 0;
 
+	virtual bool Open(const FString& Link, const TSharedRef<FCulture>& Culture) const = 0;
+
 	virtual TSharedRef< class SWidget > CreateAnchor( const FString& Link, const FString& PreviewLink = FString(), const FString& PreviewExcerptName = FString() ) const = 0;
 
 	virtual TSharedRef< class IDocumentationPage > GetPage( const FString& Link, const TSharedPtr< FParserConfiguration >& Config, const FDocumentationStyle& Style = FDocumentationStyle() ) = 0;
 
 	virtual bool PageExists(const FString& Link) const = 0;
 
+	virtual bool PageExists(const FString& Link, const TSharedRef<FCulture>& Culture) const = 0;
+
 	virtual TSharedRef< class SToolTip > CreateToolTip( const TAttribute<FText>& Text, const TSharedPtr<SWidget>& OverrideContent, const FString& Link, const FString& ExcerptName ) const = 0;
 };
