Editor: Added source control checkout or add when creating a new rich tooltip excerpt via the 'create' link

diff --git a/Engine/Source/Editor/Documentation/Documentation.Build.cs b/Engine/Source/Editor/Documentation/Documentation.Build.cs
--- a/Engine/Source/Editor/Documentation/Documentation.Build.cs
+++ b/Engine/Source/Editor/Documentation/Documentation.Build.cs
@@ -39,10 +39,13 @@ namespace UnrealBuildTool.Rules
                     "EditorStyle",
                     "UnrealEd",
 					"Analytics",
-					"SourceCodeAccess"
+					"SourceCodeAccess",
+					"SourceControl"
 				}
 			);
 
+			CircularlyReferencedDependentModules.Add("SourceControl");
+
 			DynamicallyLoadedModuleNames.AddRange(
 				new string[]
 				{
diff --git a/Engine/Source/Editor/Documentation/Private/SDocumentationToolTip.cpp b/Engine/Source/Editor/Documentation/Private/SDocumentationToolTip.cpp
--- a/Engine/Source/Editor/Documentation/Private/SDocumentationToolTip.cpp
+++ b/Engine/Source/Editor/Documentation/Private/SDocumentationToolTip.cpp
@@ -4,6 +4,7 @@
 #include "SDocumentationToolTip.h"
 #include "DocumentationLink.h"
 #include "ISourceCodeAccessModule.h"
+#include "Developer/SourceControl/Public/SourceControlHelpers.h"
 
 void SDocumentationToolTip::Construct( const FArguments& InArgs )
 {
@@ -153,15 +154,19 @@ void SDocumentationToolTip::ConstructSimpleTipContent()
 
 void SDocumentationToolTip::CreateExcerpt( FString FileSource, FString InExcerptName )
 {
-	bool NewFile = true;
+	FText CheckoutFailReason;
+	bool bNewFile = true;
+	bool bCheckoutOrAddSucceeded = true;
 	if (FPaths::FileExists(FileSource))
 	{
-		NewFile = false;
+		// Check out the existing file
+		bNewFile = false;
+		bCheckoutOrAddSucceeded = SourceControlHelpers::CheckoutOrMarkForAdd(FileSource, NSLOCTEXT("SToolTip", "DocumentationSCCActionDesc", "tool tip excerpt"), FOnPostCheckOut(), /*out*/ CheckoutFailReason);
 	}
 
 	FArchive* FileWriter = IFileManager::Get().CreateFileWriter( *FileSource, EFileWrite::FILEWRITE_Append | EFileWrite::FILEWRITE_AllowRead | EFileWrite::FILEWRITE_EvenIfReadOnly );
 
-	if ( NewFile )
+	if (bNewFile)
 	{
 		FString UdnHeader;
 		UdnHeader += "Availability:NoPublish";
@@ -191,7 +196,7 @@ void SDocumentationToolTip::CreateExcerpt( FString FileSource, FString InExcerpt
 	NewExcerpt += "]";
 	NewExcerpt += LINE_TERMINATOR;
 
-	if ( !NewFile )
+	if (!bNewFile)
 	{
 		FileWriter->Seek( FMath::Max( FileWriter->TotalSize(), (int64)0 ) );
 	}
@@ -201,9 +206,22 @@ void SDocumentationToolTip::CreateExcerpt( FString FileSource, FString InExcerpt
 	FileWriter->Close();
 	delete FileWriter;
 
+	if (bNewFile)
+	{
+		// Add the new file
+		bCheckoutOrAddSucceeded = SourceControlHelpers::CheckoutOrMarkForAdd(FileSource, NSLOCTEXT("SToolTip", "DocumentationSCCActionDesc", "tool tip excerpt"), FOnPostCheckOut(), /*out*/ CheckoutFailReason);
+	}
+
 	ISourceCodeAccessModule& SourceCodeAccessModule = FModuleManager::LoadModuleChecked<ISourceCodeAccessModule>("SourceCodeAccess");
 	SourceCodeAccessModule.GetAccessor().OpenFileAtLine(FileSource, 0);
 
+	if (!bCheckoutOrAddSucceeded)
+	{
+		FNotificationInfo Info(CheckoutFailReason);
+		Info.ExpireDuration = 3.0f;
+		FSlateNotificationManager::Get().AddNotification(Info);
+	}
+
 	ReloadDocumentation();
 }
 
