Added [CODELINK] token to UDN parser to link to source code files in tutorials and such.

diff --git a/Engine/Source/Editor/Documentation/Documentation.Build.cs b/Engine/Source/Editor/Documentation/Documentation.Build.cs
--- a/Engine/Source/Editor/Documentation/Documentation.Build.cs
+++ b/Engine/Source/Editor/Documentation/Documentation.Build.cs
@@ -40,7 +40,8 @@ namespace UnrealBuildTool.Rules
                     "UnrealEd",
 					"Analytics",
 					"SourceCodeAccess",
-					"SourceControl"
+					"SourceControl",
+                    "ContentBrowser"
 				}
 			);
 
diff --git a/Engine/Source/Editor/Documentation/Private/UDNParser.cpp b/Engine/Source/Editor/Documentation/Private/UDNParser.cpp
--- a/Engine/Source/Editor/Documentation/Private/UDNParser.cpp
+++ b/Engine/Source/Editor/Documentation/Private/UDNParser.cpp
@@ -5,6 +5,10 @@
 #include "Developer/MessageLog/Public/MessageLogModule.h"
 #include "MessageLog.h"
 #include "DocumentationLink.h"
+#include "ISourceCodeAccessModule.h"
+#include "ContentBrowserModule.h"
+#include "AssetEditorManager.h"
+#include "SourceCodeNavigation.h"
 
 #define LOCTEXT_NAMESPACE "IntroTutorials"
 
@@ -1139,6 +1143,9 @@ void FUDNParser::NavigateToLink( FString AdditionalContent )
 	static const FString HttpLinkSPecifier( TEXT( "http://" ) );
 	static const FString HttpsLinkSPecifier( TEXT( "https://" ) );
 
+	static const FString CodeLinkSpecifier(TEXT("CODELINK:"));
+	static const FString AssetLinkSpecifier(TEXT("ASSETLINK:"));
+
 	if (AdditionalContent.StartsWith(DocLinkSpecifier))
 	{
 		// external link to documentation
@@ -1156,6 +1163,16 @@ void FUDNParser::NavigateToLink( FString AdditionalContent )
 		// external link
 		FPlatformProcess::LaunchURL( *AdditionalContent, nullptr, nullptr);
 	}
+	else if (AdditionalContent.StartsWith(CodeLinkSpecifier))
+	{		
+		FString InternalLink = AdditionalContent.RightChop(CodeLinkSpecifier.Len()); 		
+		ParseCodeLink(InternalLink);
+	}	
+	else if (AdditionalContent.StartsWith(AssetLinkSpecifier))
+	{
+		FString InternalLink = AdditionalContent.RightChop(AssetLinkSpecifier.Len());
+		ParseAssetLink(InternalLink);
+	}
 	else
 	{
 		// internal link
@@ -1163,5 +1180,99 @@ void FUDNParser::NavigateToLink( FString AdditionalContent )
 	}
 }
 
+bool FUDNParser::ParseCodeLink(FString &InternalLink)
+{
+	// Tokens used by the code parsing. Details in the parse section	
+	static const FString ProjectSpecifier(TEXT("[PROJECT]"));
+	static const FString ProjectRoot(TEXT("[PROJECT]/Source/[PROJECT]/"));
+	static const FString ProjectSuffix(TEXT(".uproject"));
+
+	bool bLinkParsedOK = false;	
+	FString Path;
+	int32 Line = 0;
+	int32 Col = 0;
+
+	TArray<FString> Tokens;
+	InternalLink.ParseIntoArray(&Tokens, TEXT(","), 0);
+	int32 TokenStringsCount = Tokens.Num();
+	if (TokenStringsCount > 0)
+	{
+		Path = Tokens[0];
+	}
+	if (TokenStringsCount > 1)
+	{
+		TTypeFromString<int32>::FromString(Line, *Tokens[1]);
+	}
+	if (TokenStringsCount > 2)
+	{
+		TTypeFromString<int32>::FromString(Col, *Tokens[2]);
+	}
+
+	ISourceCodeAccessModule& SourceCodeAccessModule = FModuleManager::LoadModuleChecked<ISourceCodeAccessModule>("SourceCodeAccess");
+	ISourceCodeAccessor& SourceCodeAccessor = SourceCodeAccessModule.GetAccessor();
+	FString SolutionPath;
+
+	// If we specified generic project specified as the project name try to replace the name with the name of this project
+	if (InternalLink.Contains(ProjectSpecifier) == true)
+	{
+		FString ProjectName = TEXT("Marble");
+		// Try to extract the name of the project
+		FString ProjectPath = FPaths::GetProjectFilePath();
+		if (ProjectPath.EndsWith(ProjectSuffix))
+		{
+			int32 ProjectPathEndIndex;
+			if (ProjectPath.FindLastChar(TEXT('/'), ProjectPathEndIndex) == true)
+			{
+				ProjectName = ProjectPath.RightChop(ProjectPathEndIndex + 1);
+				ProjectName.RemoveFromEnd(*ProjectSuffix);
+			}
+		}
+		// Replace the root path with the name of this project
+		FString RebuiltPath = ProjectRoot + Path;
+		RebuiltPath.ReplaceInline(*ProjectSpecifier, *ProjectName);
+		Path = RebuiltPath;
+	}
+
+	// Finally create the complete path - project name and all
+	int32 PathEndIndex;
+	SolutionPath = FModuleManager::Get().GetSolutionFilepath();
+	if (SolutionPath.FindLastChar(TEXT('/'), PathEndIndex) == true)
+	{
+		SolutionPath = SolutionPath.LeftChop(SolutionPath.Len() - PathEndIndex - 1);
+		SolutionPath += Path;
+		bLinkParsedOK = SourceCodeAccessor.OpenFileAtLine(SolutionPath, Line, Col);
+	}
+	return bLinkParsedOK;
+}
+
+bool FUDNParser::ParseAssetLink(FString &InternalLink)
+{
+	TArray<FString> Token;
+	InternalLink.ParseIntoArray(&Token, TEXT(","), 0);
+	
+	if (Token.Num() >= 2)
+	{
+		FString Action = Token[0];
+		FString AssetName = Token[1];
+
+		UObject* RequiredObject = FindObject<UObject>(ANY_PACKAGE, *AssetName);
+		if (RequiredObject != nullptr)
+		{
+			if (Action == TEXT("EDIT"))
+			{
+				FAssetEditorManager::Get().OpenEditorForAsset(RequiredObject);
+			}
+			else
+			{
+				FContentBrowserModule& ContentBrowserModule = FModuleManager::Get().LoadModuleChecked<FContentBrowserModule>("ContentBrowser");
+				TArray<UObject*> AssetToBrowse;
+				AssetToBrowse.Add(RequiredObject);
+				ContentBrowserModule.Get().SyncBrowserToAssets(AssetToBrowse);
+			}
+		}
+	}
+	
+	return false;
+}
 
 #undef LOCTEXT_NAMESPACE
diff --git a/Engine/Source/Editor/Documentation/Private/UDNParser.h b/Engine/Source/Editor/Documentation/Private/UDNParser.h
--- a/Engine/Source/Editor/Documentation/Private/UDNParser.h
+++ b/Engine/Source/Editor/Documentation/Private/UDNParser.h
@@ -166,6 +166,27 @@ private:
 	void NavigateToLink( FString AdditionalContent );
 
 private:
+	/* 
+	 * Parses a code link embedded in the doc.
+	 * Allows us to specify files in code to link to in one of 2 ways. In both cases the last 2 parameters are line and column.
+	 * [Project based link](CODELINK:Private/[PROJECT]File.cpp, 29, 5)
+	 * This will attempt to parse the active solution name and replace instances of [PROJECT] within this so
+	 * (CODELINK:Private/[PROJECT]Ball.cpp, 29, 5) will in a project called marble will equate to
+	 * <UE4ROOT>Marble/Source/Marble/Private/MarbleBall.cpp
+	 * [Explicit link](CODELINK:Templates/TP_Rolling/Source/TP_Rolling/Private/TP_Rolling.cpp will equate to
+	 * <UE4ROOT>Templates/TP_Rolling/Source/TP_Rolling/Private/TP_Rolling.cpp	 
+	 */
+	bool ParseCodeLink(FString &InternalLink);
+		
+	/* 
+	 * Parses an asset link embedded in the doc.
+	 * Allows us to specify assets to either highlight or edit in the editor
+	 * (ASSETLINK:SELECT,MyCharacter)
+	 * This will highlight the MyCharacter asset in the content browser
+	 * (ASSETLINK:EDIT,MyCharacter)
+	 * This will edit the given asset in the appropriate editor window type.
+	 */
+	bool ParseAssetLink(FString &InternalLink);
 
 	FUDNParser( const TSharedRef< FParserConfiguration >& InConfiguration, const FDocumentationStyle& InStyle );
 	void Initialize();
