Add analytics event for when Ctrl + Alt is used to show a rich tooltip

diff --git a/Engine/Source/Editor/Documentation/Private/SDocumentationToolTip.cpp b/Engine/Source/Editor/Documentation/Private/SDocumentationToolTip.cpp
--- a/Engine/Source/Editor/Documentation/Private/SDocumentationToolTip.cpp
+++ b/Engine/Source/Editor/Documentation/Private/SDocumentationToolTip.cpp
@@ -5,6 +5,8 @@
 #include "DocumentationLink.h"
 #include "ISourceCodeAccessModule.h"
 #include "Developer/SourceControl/Public/SourceControlHelpers.h"
+#include "EngineAnalytics.h"
+#include "IAnalyticsProvider.h"
 
 void SDocumentationToolTip::Construct( const FArguments& InArgs )
 {
@@ -351,6 +353,16 @@ void SDocumentationToolTip::Tick( const FGeometry& AllottedGeometry, const doubl
 		{
 			WidgetContent->SetContent( FullTipContent.ToSharedRef() );
 			IsShowingFullTip = true;
+
+			// Analytics event
+			if (FEngineAnalytics::IsAvailable())
+			{
+				TArray<FAnalyticsEventAttribute> Params;
+				Params.Add(FAnalyticsEventAttribute(TEXT("Page"), DocumentationLink));
+				Params.Add(FAnalyticsEventAttribute(TEXT("Excerpt"), ExcerptName));
+
+				FEngineAnalytics::GetProvider().RecordEvent(TEXT("Editor.Usage.Documentation.FullTooltipShown"), Params);
+			}
 		}
 	}
 	else if ( ( IsShowingFullTip || NeedsUpdate )  && ( !ModifierKeys.IsAltDown() || !ModifierKeys.IsControlDown() ) )
