Removed dialog from Documentation Actor.
(This fixes build error also)

diff --git a/Engine/Source/Editor/Documentation/Private/Documentation.cpp b/Engine/Source/Editor/Documentation/Private/Documentation.cpp
--- a/Engine/Source/Editor/Documentation/Private/Documentation.cpp
+++ b/Engine/Source/Editor/Documentation/Private/Documentation.cpp
@@ -10,6 +10,8 @@
 #include "IAnalyticsProvider.h"
 #include "EngineAnalytics.h"
 
+#define LOCTEXT_NAMESPACE "DocumentationActor"
+
 TSharedRef< IDocumentation > FDocumentation::Create() 
 {
 	return MakeShareable( new FDocumentation() );
@@ -55,6 +57,26 @@ bool FDocumentation::Open( const FString& Link ) const
 {
 	FString DocumentationUrl;
 
+	// Warn the user if they are opening a URL
+	if (Link.StartsWith(TEXT("http")) || Link.StartsWith(TEXT("https")))
+	{
+		FText Message = LOCTEXT("OpeningURLMessage", "You are about to open an external URL. This will open your web browser. Do you want to proceed?");
+		FText URLDialog = LOCTEXT("OpeningURLTitle", "Open external link");
+
+		FSuppressableWarningDialog::FSetupInfo Info(Message, URLDialog, "SupressOpenURLWarning");
+		Info.ConfirmText = LOCTEXT("OpenURL_yes", "Yes");
+		Info.CancelText = LOCTEXT("OpenURL_no", "No");
+		FSuppressableWarningDialog OpenURLWarning(Info);
+		if (OpenURLWarning.ShowModal() == FSuppressableWarningDialog::Cancel)
+		{
+			return false;
+		}
+		else
+		{
+			FPlatformProcess::LaunchURL(*Link, nullptr, nullptr);
+		}
+	}
+
 	if (!FParse::Param(FCommandLine::Get(), TEXT("testdocs")))
 	{
 		FString OnDiskPath = FDocumentationLink::ToFilePath(Link);
@@ -64,6 +86,8 @@ bool FDocumentation::Open( const FString& Link ) const
 		}
 	}
 
+	
+	
 	if (DocumentationUrl.IsEmpty())
 	{
 		// When opening a doc website we always request the most ideal culture for our documentation.
@@ -201,4 +225,6 @@ TSharedRef< class SToolTip > FDocumentation::CreateToolTip( const TAttribute<FTe
 		[
 			DocToolTip.ToSharedRef()
 		];
-}
\ No newline at end of file
+}
+
+#undef LOCTEXT_NAMESPACE
diff --git a/Engine/Source/Runtime/Engine/Private/DocumentationActor.cpp b/Engine/Source/Runtime/Engine/Private/DocumentationActor.cpp
--- a/Engine/Source/Runtime/Engine/Private/DocumentationActor.cpp
+++ b/Engine/Source/Runtime/Engine/Private/DocumentationActor.cpp
@@ -8,8 +8,6 @@
 #endif
 #include "Dialogs.h"
 
-#define LOCTEXT_NAMESPACE "DocumentationActor"
-
 ADocumentationActor::ADocumentationActor(const class FPostConstructInitializeProperties& PCIP)
 	: Super(PCIP)
 {
@@ -30,29 +28,7 @@ bool ADocumentationActor::OpenDocumentLink() const
 	bool bOpened = false;
 
 #if WITH_EDITOR
-	// Warn the user if they are opening a URL
-	if (GetLinkType() == EDocumentationActorType::URLLink)
-	{
-		FText Message = LOCTEXT("OpeningURLMessage", "You are about to open an external URL. This will open your web browser. Do you want to proceed?");
-		FText URLDialog = LOCTEXT("OpeningURLTitle", "Open external link");
-
-		FSuppressableWarningDialog::FSetupInfo Info(Message, URLDialog, "SupressOpenURLWarning");
-		Info.ConfirmText = LOCTEXT("OpenURL_yes", "Yes");
-		Info.CancelText = LOCTEXT("OpenURL_no", "No");
-		FSuppressableWarningDialog OpenURLWarning(Info);
-		if (OpenURLWarning.ShowModal() == FSuppressableWarningDialog::Cancel)
-		{
-			bOpened = false;
-		}
-		else
-		{
-			FPlatformProcess::LaunchURL(*DocumentLink, nullptr, nullptr);
-		}
-	}
-	else
-	{
-		bOpened = IDocumentation::Get()->Open(DocumentLink);
-	}
+	bOpened = IDocumentation::Get()->Open(DocumentLink);	
 #endif //WITH_EDITOR
 
 	return bOpened;
@@ -109,4 +85,3 @@ void ADocumentationActor::UpdateLinkType()
 	}
 }
 
-#undef LOCTEXT_NAMESPACE 
