Fixed clicking on Document links in documentation actor opening 2 pages

diff --git a/Engine/Source/Editor/Documentation/Private/Documentation.cpp b/Engine/Source/Editor/Documentation/Private/Documentation.cpp
--- a/Engine/Source/Editor/Documentation/Private/Documentation.cpp
+++ b/Engine/Source/Editor/Documentation/Private/Documentation.cpp
@@ -74,6 +74,7 @@ bool FDocumentation::Open( const FString& Link ) const
 		else
 		{
 			FPlatformProcess::LaunchURL(*Link, nullptr, nullptr);
+			return true;
 		}
 	}
 
