Added information to documentation links URLs to be able to gather stats about where users are coming from when they view documentation.

diff --git a/Engine/Source/Developer/MessageLog/Private/UserInterface/SMessageLogMessageListRow.h b/Engine/Source/Developer/MessageLog/Private/UserInterface/SMessageLogMessageListRow.h
--- a/Engine/Source/Developer/MessageLog/Private/UserInterface/SMessageLogMessageListRow.h
+++ b/Engine/Source/Developer/MessageLog/Private/UserInterface/SMessageLogMessageListRow.h
@@ -85,7 +85,7 @@ private:
 #if WITH_EDITOR
 	void HandleDocsHyperlinkNavigate( FString DocumentationLink )
 	{
-		IDocumentation::Get()->Open(DocumentationLink);
+		IDocumentation::Get()->Open(DocumentationLink, FDocumentationSourceInfo(TEXT("msg_log")));
 	}
 
 	void HandleTutorialHyperlinkNavigate( FString TutorialAssetName )
diff --git a/Engine/Source/Editor/ContentBrowser/Private/AssetContextMenu.cpp b/Engine/Source/Editor/ContentBrowser/Private/AssetContextMenu.cpp
--- a/Engine/Source/Editor/ContentBrowser/Private/AssetContextMenu.cpp
+++ b/Engine/Source/Editor/ContentBrowser/Private/AssetContextMenu.cpp
@@ -1327,7 +1327,7 @@ void FAssetContextMenu::ExecuteGoToDocsForAsset(UClass* SelectedClass, const FSt
 		FString DocumentationLink = FEditorClassUtils::GetDocumentationLink(SelectedClass, ExcerptSection);
 		if (!DocumentationLink.IsEmpty())
 		{
-			IDocumentation::Get()->Open( DocumentationLink );
+			IDocumentation::Get()->Open(DocumentationLink, FDocumentationSourceInfo(TEXT("cb_docs")));
 		}
 	}
 }
diff --git a/Engine/Source/Editor/Documentation/Private/Documentation.cpp b/Engine/Source/Editor/Documentation/Private/Documentation.cpp
--- a/Engine/Source/Editor/Documentation/Private/Documentation.cpp
+++ b/Engine/Source/Editor/Documentation/Private/Documentation.cpp
@@ -27,14 +27,14 @@ FDocumentation::~FDocumentation()
 
 }
 
-bool FDocumentation::OpenHome() const
+bool FDocumentation::OpenHome(FDocumentationSourceInfo Source) const
 {
-	return Open( TEXT("%ROOT%") );
+	return Open( TEXT("%ROOT%"), Source );
 }
 
-bool FDocumentation::OpenHome(const FCultureRef& Culture) const
+bool FDocumentation::OpenHome(const FCultureRef& Culture, FDocumentationSourceInfo Source) const
 {
-	return Open(TEXT("%ROOT%"), Culture);
+	return Open(TEXT("%ROOT%"), Culture, Source);
 }
 
 bool FDocumentation::OpenAPIHome() const
@@ -53,7 +53,7 @@ bool FDocumentation::OpenAPIHome() const
 	}
 }
 
-bool FDocumentation::Open( const FString& Link ) const
+bool FDocumentation::Open(const FString& Link, FDocumentationSourceInfo Source) const
 {
 	FString DocumentationUrl;
 
@@ -83,17 +83,16 @@ bool FDocumentation::Open( const FString& Link ) const
 		FString OnDiskPath = FDocumentationLink::ToFilePath(Link);
 		if (IFileManager::Get().FileSize(*OnDiskPath) != INDEX_NONE)
 		{
-			DocumentationUrl = FDocumentationLink::ToFileUrl(Link);
+			DocumentationUrl = FDocumentationLink::ToFileUrl(Link, Source);
 		}
 	}
 
 	
-	
 	if (DocumentationUrl.IsEmpty())
 	{
 		// When opening a doc website we always request the most ideal culture for our documentation.
 		// The DNS will redirect us if necessary.
-		DocumentationUrl = FDocumentationLink::ToUrl(Link);
+		DocumentationUrl = FDocumentationLink::ToUrl(Link, Source);
 	}
 
 	if (!DocumentationUrl.IsEmpty())
@@ -109,7 +108,7 @@ bool FDocumentation::Open( const FString& Link ) const
 	return !DocumentationUrl.IsEmpty();
 }
 
-bool FDocumentation::Open(const FString& Link, const FCultureRef& Culture) const
+bool FDocumentation::Open(const FString& Link, const FCultureRef& Culture, FDocumentationSourceInfo Source) const
 {
 	FString DocumentationUrl;
 
@@ -118,13 +117,13 @@ bool FDocumentation::Open(const FString& Link, const FCultureRef& Culture) const
 		FString OnDiskPath = FDocumentationLink::ToFilePath(Link, Culture);
 		if (IFileManager::Get().FileSize(*OnDiskPath) != INDEX_NONE)
 		{
-			DocumentationUrl = FDocumentationLink::ToFileUrl(Link, Culture);
+			DocumentationUrl = FDocumentationLink::ToFileUrl(Link, Culture, Source);
 		}
 	}
 
 	if (DocumentationUrl.IsEmpty())
 	{
-		DocumentationUrl = FDocumentationLink::ToUrl(Link, Culture);
+		DocumentationUrl = FDocumentationLink::ToUrl(Link, Culture, Source);
 	}
 
 	if (!DocumentationUrl.IsEmpty())
diff --git a/Engine/Source/Editor/Documentation/Private/Documentation.h b/Engine/Source/Editor/Documentation/Private/Documentation.h
--- a/Engine/Source/Editor/Documentation/Private/Documentation.h
+++ b/Engine/Source/Editor/Documentation/Private/Documentation.h
@@ -13,15 +13,15 @@ public:
 
 	virtual ~FDocumentation();
 
-	virtual bool OpenHome() const override;
+	virtual bool OpenHome(FDocumentationSourceInfo Source = FDocumentationSourceInfo()) const override;
 
-	virtual bool OpenHome(const FCultureRef& Culture) const override;
+	virtual bool OpenHome(const FCultureRef& Culture, FDocumentationSourceInfo Source = FDocumentationSourceInfo()) const override;
 
 	virtual bool OpenAPIHome() const override;
 
-	virtual bool Open( const FString& Link ) const override;
+	virtual bool Open( const FString& Link, FDocumentationSourceInfo Source = FDocumentationSourceInfo() ) const override;
 
-	virtual bool Open(const FString& Link, const FCultureRef& Culture) const override;
+	virtual bool Open( const FString& Link, const FCultureRef& Culture, FDocumentationSourceInfo Source = FDocumentationSourceInfo() ) const override;
 
 	virtual TSharedRef< SWidget > CreateAnchor( const TAttribute<FString>& Link, const FString& PreviewLink = FString(), const FString& PreviewExcerptName = FString() ) const override;
 
diff --git a/Engine/Source/Editor/Documentation/Private/DocumentationLink.h b/Engine/Source/Editor/Documentation/Private/DocumentationLink.h
--- a/Engine/Source/Editor/Documentation/Private/DocumentationLink.h
+++ b/Engine/Source/Editor/Documentation/Private/DocumentationLink.h
@@ -32,20 +32,23 @@ public:
 		return Url;
 	}
 
-	static FString ToUrl(const FString& Link)
+	static FString ToUrl(const FString& Link, FDocumentationSourceInfo const& Source)
 	{
-		return ToUrl(Link, FInternationalization::Get().GetCurrentCulture());
+		return ToUrl(Link, FInternationalization::Get().GetCurrentCulture(), Source);
 	}
 
-	static FString ToUrl(const FString& Link, const FCultureRef& Culture)
+	static FString ToUrl(const FString& Link, const FCultureRef& Culture, FDocumentationSourceInfo const& Source)
 	{
 		FString Path;
 		FString Anchor;
-		SplitLink( Link, Path, Anchor );
+		FString QueryString;
+		SplitLink(Link, Path, QueryString, Anchor);
 
 		const FString PartialPath = FString::Printf(TEXT("%s%s/index.html"), *(Culture->GetUnrealLegacyThreeLetterISOLanguageName()), *Path);
 		
-		return GetUrlRoot() + PartialPath + Anchor;
+		AddSourceInfoToQueryString(QueryString, Source);
+
+		return GetUrlRoot() + PartialPath + QueryString + Anchor;
 	}
 
 	static FString ToFilePath( const FString& Link )
@@ -74,13 +77,14 @@ public:
 	{
 		FString Path;
 		FString Anchor;
-		SplitLink(Link, Path, Anchor);
+		FString QueryString;
+		SplitLink(Link, Path, QueryString, Anchor);
 
 		const FString PartialPath = FString::Printf(TEXT("%s%s/index.html"), *(Culture->GetUnrealLegacyThreeLetterISOLanguageName()), *Path);
 		return FString::Printf(TEXT("%sDocumentation/HTML/%s"), *FPaths::ConvertRelativePathToFull(FPaths::EngineDir()), *PartialPath);
 	}
 
-	static FString ToFileUrl( const FString& Link )
+	static FString ToFileUrl(const FString& Link, FDocumentationSourceInfo const& SourceInfo)
 	{
 		FInternationalization& I18N = FInternationalization::Get();
 
@@ -100,16 +104,19 @@ public:
 			}
 		}
 
-		return ToFileUrl(Link, Culture);
+		return ToFileUrl(Link, Culture, SourceInfo);
 	}
 
-	static FString ToFileUrl(const FString& Link, const FCultureRef& Culture)
+	static FString ToFileUrl(const FString& Link, const FCultureRef& Culture, FDocumentationSourceInfo const& SourceInfo)
 	{
 		FString Path;
 		FString Anchor;
-		SplitLink(Link, Path, Anchor);
+		FString QueryString;
+		SplitLink(Link, Path, QueryString, Anchor);
+
+		AddSourceInfoToQueryString(QueryString, SourceInfo);
 
-		return FString::Printf(TEXT("file:///%s%s"), *ToFilePath(Link, Culture), *Anchor);
+		return FString::Printf(TEXT("file:///%s%s%s"), *ToFilePath(Link, Culture), *QueryString, *Anchor);
 	}
 
 	static FString ToSourcePath(const FString& Link)
@@ -138,7 +145,8 @@ public:
 	{
 		FString Path;
 		FString Anchor;
-		SplitLink( Link, Path, Anchor );
+		FString QueryString;
+		SplitLink(Link, Path, QueryString, Anchor);
 
 		const FString FullDirectoryPath = FPaths::EngineDir() + TEXT( "Documentation/Source" ) + Path + "/";
 
@@ -157,7 +165,23 @@ public:
 		return FString::Printf(TEXT("%s%s.%s.udn"), *FullDirectoryPath, *Category, *(Culture->GetUnrealLegacyThreeLetterISOLanguageName()));
 	}
 
-	static void SplitLink( const FString& Link, /*OUT*/ FString& Path, /*OUT*/ FString& Anchor )
+private:
+	static void AddSourceInfoToQueryString(FString& QueryString, FDocumentationSourceInfo const& Info)
+	{
+		if (Info.IsEmpty() == false)
+		{
+			if (QueryString.IsEmpty())
+			{
+				QueryString = FString::Printf(TEXT("?utm_source=%s&utm_medium=%s&utm_campaign=%s"), *Info.Source, *Info.Medium, *Info.Campaign);
+			}
+			else
+			{
+				QueryString = FString::Printf(TEXT("%s&utm_source=%s&utm_medium=%s&utm_campaign=%s"), *QueryString, *Info.Source, *Info.Medium, *Info.Campaign);
+			}
+		}
+	}
+	
+	static void SplitLink( const FString& Link, /*OUT*/ FString& Path, /*OUT*/ FString& QueryString, /*OUT*/ FString& Anchor )
 	{
 		FString CleanedLink = Link;
 		CleanedLink.Trim();
@@ -167,15 +191,18 @@ public:
 		{
 			Path.Empty();
 			Anchor.Empty();
+			QueryString.Empty();
 		}
 		else
 		{
-			if ( !CleanedLink.Split( TEXT("#"), &Path, &Anchor ) )
+			FString PathAndQueryString;
+			if ( !CleanedLink.Split( TEXT("#"), &PathAndQueryString, &Anchor ) )
 			{
-				Path = CleanedLink;
+				PathAndQueryString = CleanedLink;
 			}
 			else if ( !Anchor.IsEmpty() )
 			{
+				// ensure leading #
 				Anchor = FString( TEXT("#") ) + Anchor;
 			}
 
@@ -184,14 +211,25 @@ public:
 				Anchor = Anchor.Left( Anchor.Len() - 1 );
 			}
 
-			if ( Path.EndsWith( TEXT("/") ) )
+			if ( PathAndQueryString.EndsWith( TEXT("/") ) )
 			{
-				Path = Path.Left( Path.Len() - 1 );
+				PathAndQueryString = PathAndQueryString.Left(PathAndQueryString.Len() - 1);
 			}
 
-			if ( !Path.IsEmpty() && !Path.StartsWith( TEXT("/") ) )
+			if ( !PathAndQueryString.IsEmpty() && !PathAndQueryString.StartsWith( TEXT("/") ) )
+			{
+				PathAndQueryString = FString(TEXT("/")) + PathAndQueryString;
+			}
+
+			// split path and query string
+			if (!PathAndQueryString.Split(TEXT("?"), &Path, &QueryString))
+			{
+				Path = PathAndQueryString;
+			}
+			else if (!QueryString.IsEmpty())
 			{
-				Path = FString( TEXT("/") ) + Path;
+				// ensure leading ?
+				QueryString = FString(TEXT("?")) + QueryString;
 			}
 		}
 	}
diff --git a/Engine/Source/Editor/Documentation/Private/SDocumentationAnchor.cpp b/Engine/Source/Editor/Documentation/Private/SDocumentationAnchor.cpp
--- a/Engine/Source/Editor/Documentation/Private/SDocumentationAnchor.cpp
+++ b/Engine/Source/Editor/Documentation/Private/SDocumentationAnchor.cpp
@@ -59,6 +59,6 @@ const FSlateBrush* SDocumentationAnchor::GetButtonImage() const
 
 FReply SDocumentationAnchor::OnClicked() const
 {
-	IDocumentation::Get()->Open( Link.Get() );
+	IDocumentation::Get()->Open(Link.Get(), FDocumentationSourceInfo(TEXT("doc_anchors")));
 	return FReply::Handled();
 }
diff --git a/Engine/Source/Editor/Documentation/Private/SDocumentationToolTip.cpp b/Engine/Source/Editor/Documentation/Private/SDocumentationToolTip.cpp
--- a/Engine/Source/Editor/Documentation/Private/SDocumentationToolTip.cpp
+++ b/Engine/Source/Editor/Documentation/Private/SDocumentationToolTip.cpp
@@ -274,7 +274,7 @@ void SDocumentationToolTip::ConstructFullTipContent()
 				{
 					static void OpenLink( FString Link )
 					{
-						if ( !IDocumentation::Get()->Open( Link ) )
+						if (!IDocumentation::Get()->Open(Link, FDocumentationSourceInfo(TEXT("rich_tooltips"))))
 						{
 							FNotificationInfo Info( NSLOCTEXT("SToolTip", "FailedToOpenLink", "Failed to Open Link") );
 							FSlateNotificationManager::Get().AddNotification(Info);
diff --git a/Engine/Source/Editor/Documentation/Private/UDNParser.cpp b/Engine/Source/Editor/Documentation/Private/UDNParser.cpp
--- a/Engine/Source/Editor/Documentation/Private/UDNParser.cpp
+++ b/Engine/Source/Editor/Documentation/Private/UDNParser.cpp
@@ -40,7 +40,7 @@ TSharedRef< FUDNParser > FUDNParser::Create( const TSharedPtr< FParserConfigurat
 		{
 			static void OpenLink( const FString& Link )
 			{
-				if ( !IDocumentation::Get()->Open( Link ) )
+				if ( !IDocumentation::Get()->Open( Link, FDocumentationSourceInfo(TEXT("udn_parser")) ) )
 				{
 					FNotificationInfo Info( NSLOCTEXT("FUDNParser", "FailedToOpenLink", "Failed to Open Link") );
 					FSlateNotificationManager::Get().AddNotification(Info);
@@ -1220,7 +1220,7 @@ void FUDNParser::NavigateToLink( FString AdditionalContent )
 	{
 		// external link to documentation
 		FString DocLink = AdditionalContent.RightChop(DocLinkSpecifier.Len());
-		IDocumentation::Get()->Open(DocLink);
+		IDocumentation::Get()->Open(DocLink, FDocumentationSourceInfo(TEXT("udn_parser")));
 	}
 	else if ( AdditionalContent.StartsWith( TutorialLinkSpecifier ) )
 	{
diff --git a/Engine/Source/Editor/Documentation/Public/IDocumentation.h b/Engine/Source/Editor/Documentation/Public/IDocumentation.h
--- a/Engine/Source/Editor/Documentation/Public/IDocumentation.h
+++ b/Engine/Source/Editor/Documentation/Public/IDocumentation.h
@@ -131,6 +131,30 @@ private:
 	FParserConfiguration() {}
 };
 
+
+struct FDocumentationSourceInfo
+{
+public:
+	FString Source;
+	FString Medium;
+	FString Campaign;
+
+	FDocumentationSourceInfo() {};
+	FDocumentationSourceInfo(FString const& InCampaign) 
+		: Source(TEXT("editor")), Medium(TEXT("docs")), Campaign(InCampaign)
+	{};
+	FDocumentationSourceInfo(FString const& InSource, FString const& InMedium, FString const& InCampaign)
+		: Source(InSource), Medium(InMedium), Campaign(InCampaign)
+	{}
+
+	/** Returns true if there is NO valid source info in the struct, false otherwise. */
+	bool IsEmpty() const
+	{
+		return Campaign.IsEmpty() && Source.IsEmpty() && Medium.IsEmpty();
+	}
+};
+
+
 class IDocumentation
 {
 public:
@@ -148,15 +172,15 @@ public:
 
 public:
 
-	virtual bool OpenHome() const = 0;
+	virtual bool OpenHome(FDocumentationSourceInfo Source = FDocumentationSourceInfo()) const = 0;
 
-	virtual bool OpenHome(const FCultureRef& Culture) const = 0;
+	virtual bool OpenHome(const FCultureRef& Culture, FDocumentationSourceInfo Source = FDocumentationSourceInfo()) const = 0;
 
 	virtual bool OpenAPIHome() const = 0;
 
-	virtual bool Open( const FString& Link ) const = 0;
+	virtual bool Open(const FString& Link, FDocumentationSourceInfo Source = FDocumentationSourceInfo()) const = 0;
 
-	virtual bool Open(const FString& Link, const FCultureRef& Culture) const = 0;
+	virtual bool Open(const FString& Link, const FCultureRef& Culture, FDocumentationSourceInfo Source = FDocumentationSourceInfo()) const = 0;
 
 	virtual TSharedRef< class SWidget > CreateAnchor( const TAttribute<FString>& Link, const FString& PreviewLink = FString(), const FString& PreviewExcerptName = FString() ) const = 0;
 
diff --git a/Engine/Source/Editor/IntroTutorials/Private/TutorialText.cpp b/Engine/Source/Editor/IntroTutorials/Private/TutorialText.cpp
--- a/Engine/Source/Editor/IntroTutorials/Private/TutorialText.cpp
+++ b/Engine/Source/Editor/IntroTutorials/Private/TutorialText.cpp
@@ -90,7 +90,7 @@ static void OnDocLinkClicked(const FSlateHyperlinkRun::FMetadata& Metadata)
 			FEngineAnalytics::GetProvider().RecordEvent( TEXT("Rocket.Tutorials.DocLinkClicked"), EventAttributes );
 		}
 
-		IDocumentation::Get()->Open(*Url);
+		IDocumentation::Get()->Open(*Url, FDocumentationSourceInfo(TEXT("tutorials")));
 	}
 }
 
diff --git a/Engine/Source/Editor/LevelEditor/Private/LevelEditorActions.cpp b/Engine/Source/Editor/LevelEditor/Private/LevelEditorActions.cpp
--- a/Engine/Source/Editor/LevelEditor/Private/LevelEditorActions.cpp
+++ b/Engine/Source/Editor/LevelEditor/Private/LevelEditorActions.cpp
@@ -113,7 +113,7 @@ bool FLevelEditorActionCallbacks::DefaultCanExecuteAction()
 
 void FLevelEditorActionCallbacks::BrowseDocumentation()
 {
-	IDocumentation::Get()->OpenHome();
+	IDocumentation::Get()->OpenHome(FDocumentationSourceInfo(TEXT("help_menu")));
 }
 
 void FLevelEditorActionCallbacks::BrowseAPIReference()
@@ -992,7 +992,7 @@ void FLevelEditorActionCallbacks::GoToDocsForActor_Clicked()
 		FString DocumentationLink = FEditorClassUtils::GetDocumentationLink(SelectedActorInfo.SelectionClass);
 		if (!DocumentationLink.IsEmpty())
 		{
-			IDocumentation::Get()->Open( DocumentationLink );
+			IDocumentation::Get()->Open( DocumentationLink, FDocumentationSourceInfo(TEXT("rightclick_viewdoc")) );
 		}
 	}
 }
diff --git a/Engine/Source/Editor/MaterialEditor/Private/MaterialEditor.cpp b/Engine/Source/Editor/MaterialEditor/Private/MaterialEditor.cpp
--- a/Engine/Source/Editor/MaterialEditor/Private/MaterialEditor.cpp
+++ b/Engine/Source/Editor/MaterialEditor/Private/MaterialEditor.cpp
@@ -2409,7 +2409,7 @@ void FMaterialEditor::OnGoToDocumentation()
 	FString DocumentationLink = GetDocLinkForSelectedNode();
 	if (!DocumentationLink.IsEmpty())
 	{
-		IDocumentation::Get()->Open(DocumentationLink);
+		IDocumentation::Get()->Open(DocumentationLink, FDocumentationSourceInfo(TEXT("rightclick_matnode")));
 	}
 }
 
diff --git a/Engine/Source/Editor/UnrealEd/Private/AboutScreen.cpp b/Engine/Source/Editor/UnrealEd/Private/AboutScreen.cpp
--- a/Engine/Source/Editor/UnrealEd/Private/AboutScreen.cpp
+++ b/Engine/Source/Editor/UnrealEd/Private/AboutScreen.cpp
@@ -156,7 +156,7 @@ const FSlateBrush* SAboutScreen::GetFacebookButtonBrush() const
 
 FReply SAboutScreen::OnUE4ButtonClicked()
 {
-	IDocumentation::Get()->OpenHome();
+	IDocumentation::Get()->OpenHome(FDocumentationSourceInfo(TEXT("logo_docs")));
 	return FReply::Handled();
 }
 
diff --git a/Engine/Source/Editor/UnrealEd/Private/Toolkits/AssetEditorToolkit.cpp b/Engine/Source/Editor/UnrealEd/Private/Toolkits/AssetEditorToolkit.cpp
--- a/Engine/Source/Editor/UnrealEd/Private/Toolkits/AssetEditorToolkit.cpp
+++ b/Engine/Source/Editor/UnrealEd/Private/Toolkits/AssetEditorToolkit.cpp
@@ -528,7 +528,7 @@ void FAssetEditorToolkit::FindInContentBrowser_Execute()
 
 void FAssetEditorToolkit::BrowseDocumentation_Execute() const
 {
-	IDocumentation::Get()->Open(GetDocumentationLink());
+	IDocumentation::Get()->Open(GetDocumentationLink(), FDocumentationSourceInfo(TEXT("help_menu_asset")));
 }
 
 void FAssetEditorToolkit::ViewReferences_Execute()
diff --git a/Engine/Source/Runtime/Engine/Private/DocumentationActor.cpp b/Engine/Source/Runtime/Engine/Private/DocumentationActor.cpp
--- a/Engine/Source/Runtime/Engine/Private/DocumentationActor.cpp
+++ b/Engine/Source/Runtime/Engine/Private/DocumentationActor.cpp
@@ -30,7 +30,7 @@ bool ADocumentationActor::OpenDocumentLink() const
 	bool bOpened = false;
 
 #if WITH_EDITOR
-	bOpened = IDocumentation::Get()->Open(DocumentLink);	
+	bOpened = IDocumentation::Get()->Open(DocumentLink, FDocumentationSourceInfo(TEXT("doc_actors")));
 #endif //WITH_EDITOR
 
 	return bOpened;
