[INTEGRATE] Change 2414097 by [REDACTED_PERSON][REDACTED_HANDLE] on 2015/01/21 13:43:14

	New look and feel for Slate tool-tips
	- Tool-tips are now black on white so they stand out from the rest of the editor
	- Increased tool-tip font size by one point; increased the amount of padding
	- Only regular text-only tool-tips and documentation tool-tips are affected.  We may update the various customized tool-tips later on.

diff --git a/Engine/Source/Editor/Documentation/Private/Documentation.cpp b/Engine/Source/Editor/Documentation/Private/Documentation.cpp
--- a/Engine/Source/Editor/Documentation/Private/Documentation.cpp
+++ b/Engine/Source/Editor/Documentation/Private/Documentation.cpp
@@ -222,6 +222,11 @@ TSharedRef< class SToolTip > FDocumentation::CreateToolTip(const TAttribute<FTex
 	
 	return SNew( SToolTip )
 		.IsInteractive( DocToolTip.ToSharedRef(), &SDocumentationToolTip::IsInteractive )
+
+		// Emulate text-only tool-tip styling that SToolTip uses when no custom content is supplied.  We want documentation tool-tips to 
+		// be styled just like text-only tool-tips
+		.BorderImage( FCoreStyle::Get().GetBrush("ToolTip.BrightBackground") )
+		.TextMargin(FMargin(11.0f))
 		[
 			DocToolTip.ToSharedRef()
 		];
diff --git a/Engine/Source/Editor/Documentation/Private/SDocumentationToolTip.cpp b/Engine/Source/Editor/Documentation/Private/SDocumentationToolTip.cpp
--- a/Engine/Source/Editor/Documentation/Private/SDocumentationToolTip.cpp
+++ b/Engine/Source/Editor/Documentation/Private/SDocumentationToolTip.cpp
@@ -15,6 +15,9 @@ void SDocumentationToolTip::Construct( const FArguments& InArgs )
 {
 	TextContent = InArgs._Text;
 	StyleInfo = FEditorStyle::GetWidgetStyle<FTextBlockStyle>(InArgs._Style);
+	SubduedStyleInfo = FEditorStyle::GetWidgetStyle<FTextBlockStyle>(InArgs._SubduedStyle);
+	HyperlinkTextStyleInfo = FEditorStyle::GetWidgetStyle<FTextBlockStyle>(InArgs._HyperlinkTextStyle);
+	HyperlinkButtonStyleInfo = FEditorStyle::GetWidgetStyle<FButtonStyle>(InArgs._HyperlinkButtonStyle);
 	ColorAndOpacity = InArgs._ColorAndOpacity;
 	DocumentationLink = InArgs._DocumentationLink;
 	IsDisplayingDocumentationLink = false;
@@ -102,12 +105,12 @@ void SDocumentationToolTip::AddDocumentation(TSharedPtr< SVerticalBox > Vertical
 
 			VerticalBox->AddSlot()
 			.AutoHeight()
-			.HAlign( HAlign_Right )
+			.Padding(0, 5, 0, 0)
+			.HAlign( HAlign_Center )
 			[
 				SNew( STextBlock )
-				.ColorAndOpacity( FSlateColor::UseSubduedForeground() )
 				.Text( FText::FromString(DocumentationLink + OptionalExcerptName) )
-				.TextStyle( &StyleInfo )
+				.TextStyle( &SubduedStyleInfo )
 			];
 		}
 
@@ -130,12 +133,12 @@ void SDocumentationToolTip::AddDocumentation(TSharedPtr< SVerticalBox > Vertical
 
 			VerticalBox->AddSlot()
 			.AutoHeight()
-			.HAlign( HAlign_Right )
+			.HAlign( HAlign_Center )
+			.Padding(0, 5, 0, 0)
 			[
 				SNew( STextBlock )
-				.ColorAndOpacity( FSlateColor::UseSubduedForeground() )
+				.TextStyle( &SubduedStyleInfo )
 				.Text( FText::Format( NSLOCTEXT( "SToolTip", "AdvancedToolTipMessage", "hold {0} for more" ), KeyboardShortcut) )
-				.TextStyle( &StyleInfo )
 			];
 		}
 		else
@@ -150,10 +153,13 @@ void SDocumentationToolTip::AddDocumentation(TSharedPtr< SVerticalBox > Vertical
 
 				VerticalBox->AddSlot()
 				.AutoHeight()
-				.HAlign( HAlign_Right )
+				.Padding(0, 5, 0, 0)
+				.HAlign( HAlign_Center )
 				[
 					SNew( SHyperlink )
 					.Text( NSLOCTEXT( "SToolTip", "EditDocumentationMessage_Create", "create" ) )
+					.TextStyle( &HyperlinkTextStyleInfo )
+					.UnderlineStyle( &HyperlinkButtonStyleInfo )
 					.OnNavigate( this, &SDocumentationToolTip::CreateExcerpt, DocPath, ExcerptName )
 				];
 			}
@@ -296,6 +302,8 @@ void SDocumentationToolTip::ConstructFullTipContent()
 				[
 					SNew( SHyperlink )
 						.Text( NSLOCTEXT( "SToolTip", "GoToFullDocsLinkMessage", "see full documentation" ) )
+						.TextStyle( &HyperlinkTextStyleInfo )
+						.UnderlineStyle( &HyperlinkButtonStyleInfo )
 						.OnNavigate_Static( &Local::OpenLink, *FullDocumentationLink )
 				];
 			}
@@ -313,10 +321,12 @@ void SDocumentationToolTip::ConstructFullTipContent()
 
 				Box->AddSlot()
 				.AutoHeight()
-				.HAlign( HAlign_Right )
+				.HAlign( HAlign_Center )
 				[
 					SNew( SHyperlink )
 						.Text( NSLOCTEXT( "SToolTip", "EditDocumentationMessage_Edit", "edit" ) )
+						.TextStyle( &HyperlinkTextStyleInfo )
+						.UnderlineStyle( &HyperlinkButtonStyleInfo )
 						.OnNavigate_Static(&Local::EditSource, FPaths::ConvertRelativePathToFull(FDocumentationLink::ToSourcePath(DocumentationLink, FInternationalization::Get().GetCurrentCulture())), Excerpts[ExcerptIndex].LineNumber)
 				];
 			}
diff --git a/Engine/Source/Editor/Documentation/Public/SDocumentationToolTip.h b/Engine/Source/Editor/Documentation/Public/SDocumentationToolTip.h
--- a/Engine/Source/Editor/Documentation/Public/SDocumentationToolTip.h
+++ b/Engine/Source/Editor/Documentation/Public/SDocumentationToolTip.h
@@ -10,7 +10,10 @@ public:
 	SLATE_BEGIN_ARGS( SDocumentationToolTip )
 		: _Text()
 		, _Style( TEXT("Documentation.SDocumentationTooltip") )
-		, _ColorAndOpacity( FLinearColor( 1.0f, 1.0f, 1.0f, 1.0f ) )
+		, _SubduedStyle( TEXT("Documentation.SDocumentationTooltipSubdued") )
+		, _HyperlinkTextStyle( TEXT("Documentation.SDocumentationTooltipHyperlinkText") )
+		, _HyperlinkButtonStyle( TEXT("Documentation.SDocumentationTooltipHyperlinkButton") )
+		, _ColorAndOpacity( FLinearColor::Black )
 		, _AddDocumentation( true )
 		, _DocumentationMargin(0)
 		, _Content()
@@ -22,6 +25,15 @@ public:
 		/** The text style to use for this tool tip */
 		SLATE_ARGUMENT( FName, Style )
 		
+		/** The text style to use for subdued footer text in this tool tip */
+		SLATE_ARGUMENT( FName, SubduedStyle )
+		
+		/** The text style to use for hyperlinks in this tool tip */
+		SLATE_ARGUMENT( FName, HyperlinkTextStyle )
+
+		/** Hyperlink button style */
+		SLATE_ARGUMENT( FName, HyperlinkButtonStyle )
+
 		/** Font color and opacity */
 		SLATE_ATTRIBUTE( FSlateColor, ColorAndOpacity )
 
@@ -81,6 +93,9 @@ private:
 	TAttribute< FText > TextContent;
 	TSharedPtr< SWidget > OverrideContent;
 	FTextBlockStyle StyleInfo;
+	FTextBlockStyle SubduedStyleInfo;
+	FTextBlockStyle HyperlinkTextStyleInfo;
+	FButtonStyle HyperlinkButtonStyleInfo;
 	TAttribute< FSlateColor > ColorAndOpacity;
 
 	/** The link to the documentation */
diff --git a/Engine/Source/Editor/EditorStyle/Private/SlateEditorStyle.cpp b/Engine/Source/Editor/EditorStyle/Private/SlateEditorStyle.cpp
--- a/Engine/Source/Editor/EditorStyle/Private/SlateEditorStyle.cpp
+++ b/Engine/Source/Editor/EditorStyle/Private/SlateEditorStyle.cpp
@@ -2535,27 +2535,47 @@ void FSlateEditorStyle::FStyle::SetupDockingStyles()
 void FSlateEditorStyle::FStyle::SetupTutorialStyles()
 {
 	// Documentation tooltip defaults
+	const FSlateColor HyperlinkColor( FLinearColor( 0.1f, 0.1f, 0.5f ) );
 	{
-		const FTextBlockStyle DocumentationTooltipText = FTextBlockStyle(NormalText)
-			.SetFont(TTF_CORE_FONT("Fonts/Roboto-Regular", 8));
+		const FTextBlockStyle DocumentationTooltipText = FTextBlockStyle( NormalText )
+			.SetFont( TTF_CORE_FONT( "Fonts/Roboto-Regular", 9 ) )
+			.SetColorAndOpacity( FLinearColor::Black );
 		Set("Documentation.SDocumentationTooltip", FTextBlockStyle(DocumentationTooltipText));
+
+		const FTextBlockStyle DocumentationTooltipTextSubdued = FTextBlockStyle( NormalText )
+			.SetFont( TTF_CORE_FONT( "Fonts/Roboto-Regular", 8 ) )
+			.SetColorAndOpacity( FLinearColor( 0.3f, 0.3f, 0.3f ) );
+		Set("Documentation.SDocumentationTooltipSubdued", FTextBlockStyle(DocumentationTooltipTextSubdued));
+
+		const FTextBlockStyle DocumentationTooltipHyperlinkText = FTextBlockStyle( NormalText )
+			.SetFont( TTF_CORE_FONT( "Fonts/Roboto-Regular", 8 ) )
+			.SetColorAndOpacity( HyperlinkColor );
+		Set("Documentation.SDocumentationTooltipHyperlinkText", FTextBlockStyle(DocumentationTooltipHyperlinkText));
+
+		const FButtonStyle DocumentationTooltipHyperlinkButton = FButtonStyle()
+				.SetNormal(BORDER_BRUSH( "Old/HyperlinkDotted", FMargin(0,0,0,3/16.0f), HyperlinkColor ) )
+				.SetPressed(FSlateNoResource())
+				.SetHovered(BORDER_BRUSH( "Old/HyperlinkUnderline", FMargin(0,0,0,3/16.0f), HyperlinkColor ) );
+		Set("Documentation.SDocumentationTooltipHyperlinkButton", FButtonStyle(DocumentationTooltipHyperlinkButton));
 	}
 
 
 	// Documentation defaults
 	const FTextBlockStyle DocumentationText = FTextBlockStyle(NormalText)
+		.SetColorAndOpacity( FLinearColor::Black )
 		.SetFont(TTF_CORE_FONT( "Fonts/Roboto-Regular", 11 ));
 
 	const FTextBlockStyle DocumentationHyperlinkText = FTextBlockStyle(DocumentationText)
-		.SetColorAndOpacity(FLinearColor::Blue);
+		.SetColorAndOpacity( HyperlinkColor );
 
 	const FTextBlockStyle DocumentationHeaderText = FTextBlockStyle(NormalText)
+		.SetColorAndOpacity( FLinearColor::Black )
 		.SetFont(TTF_FONT("Fonts/Roboto-Black", 32));
 
 	const FButtonStyle DocumentationHyperlinkButton = FButtonStyle()
-			.SetNormal(BORDER_BRUSH( "Old/HyperlinkDotted", FMargin(0,0,0,3/16.0f), FSlateColor::UseForeground()))
+			.SetNormal(BORDER_BRUSH( "Old/HyperlinkDotted", FMargin(0,0,0,3/16.0f), HyperlinkColor ) )
 			.SetPressed(FSlateNoResource())
-			.SetHovered(BORDER_BRUSH( "Old/HyperlinkUnderline", FMargin(0,0,0,3/16.0f), FSlateColor::UseForeground()));
+			.SetHovered(BORDER_BRUSH( "Old/HyperlinkUnderline", FMargin(0,0,0,3/16.0f), HyperlinkColor ) );
 
 	// Documentation
 	{
diff --git a/Engine/Source/Runtime/Slate/Private/Widgets/SToolTip.cpp b/Engine/Source/Runtime/Slate/Private/Widgets/SToolTip.cpp
--- a/Engine/Source/Runtime/Slate/Private/Widgets/SToolTip.cpp
+++ b/Engine/Source/Runtime/Slate/Private/Widgets/SToolTip.cpp
@@ -36,25 +36,41 @@ void SToolTip::SetContentWidget(const TSharedRef<SWidget>& InContentWidget)
 		WidgetContent = InContentWidget;
 	}
 
-	ToolTipContent = (!WidgetContent.IsValid())
-		? StaticCastSharedRef<SWidget>(
-		SNew(STextBlock)
-		.Text(TextContent)
-		.Font(Font)
-		.ColorAndOpacity(ColorAndOpacity)
-		.WrapTextAt_Static(&SToolTip::GetToolTipWrapWidth)
-		)
-		: InContentWidget;
+	TSharedPtr< SWidget > PinnedWidgetContent = WidgetContent.Pin();
+	if( PinnedWidgetContent.IsValid() )
+	{
+		// Tool-tip with entirely custom content.  We'll create a border with some padding (as customized by the user), then
+		// embed their custom widget right inside the border.  This tool-tip currently has a different styling than tool-tips
+		// that only contain text.
+		ChildSlot
+		[
+			SNew(SBorder)
+			.BorderImage(BorderImage)
+			.Padding(TextMargin)
+			[
+				PinnedWidgetContent.ToSharedRef()
+			]
+		];
 
-	ChildSlot
-	[
-		SNew(SBorder)
-		.BorderImage(BorderImage)
-		.Padding(TextMargin)
+	}
+	else
+	{
+		// Text-only tool-tip.  This tool-tip currently has a different styling than tool-tips with custom content.  We always want basic
+		// text tool-tips to look consistent.
+		ChildSlot
 		[
-			ToolTipContent.ToSharedRef()
-		]
-	];
+			SNew(SBorder)
+			.BorderImage( FCoreStyle::Get().GetBrush("ToolTip.BrightBackground") )
+			.Padding(FMargin(11.0f))
+			[
+				SNew(STextBlock)
+				.Text(TextContent)
+				.Font(FCoreStyle::Get().GetFontStyle("ToolTip.LargerFont"))
+				.ColorAndOpacity( FLinearColor::Black )
+				.WrapTextAt_Static(&SToolTip::GetToolTipWrapWidth)
+			]
+		];
+	}
 }
 
 
diff --git a/Engine/Source/Runtime/SlateCore/Private/Styling/CoreStyle.cpp b/Engine/Source/Runtime/SlateCore/Private/Styling/CoreStyle.cpp
--- a/Engine/Source/Runtime/SlateCore/Private/Styling/CoreStyle.cpp
+++ b/Engine/Source/Runtime/SlateCore/Private/Styling/CoreStyle.cpp
@@ -508,6 +508,9 @@ TSharedRef<ISlateStyle> FCoreStyle::Create( const FName& InStyleSetName )
 	{
 		Style->Set("ToolTip.Font", TTF_FONT("Fonts/Roboto-Regular", 8));
 		Style->Set("ToolTip.Background", new BOX_BRUSH("Old/ToolTip_Background", FMargin(8.0f / 64.0f)));
+
+		Style->Set("ToolTip.LargerFont", TTF_FONT("Fonts/Roboto-Regular", 9));
+		Style->Set("ToolTip.BrightBackground", new BOX_BRUSH("Old/ToolTip_BrightBackground", FMargin(8.0f / 64.0f)));
 	}
 
 	// SBorder defaults...
