Swap Cmd+Alt to Command+Option to match modern Apple keyboards.

diff --git a/Engine/Source/Editor/Documentation/Private/SDocumentationToolTip.cpp b/Engine/Source/Editor/Documentation/Private/SDocumentationToolTip.cpp
--- a/Engine/Source/Editor/Documentation/Private/SDocumentationToolTip.cpp
+++ b/Engine/Source/Editor/Documentation/Private/SDocumentationToolTip.cpp
@@ -134,7 +134,7 @@ void SDocumentationToolTip::AddDocumentation(TSharedPtr< SVerticalBox > Vertical
 
 		if ( DocumentationPage->HasExcerpt( ExcerptName ) )
 		{
-			FText MacShortcut = NSLOCTEXT("SToolTip", "MacRichTooltipShortcut", "(Cmd + Alt)");
+			FText MacShortcut = NSLOCTEXT("SToolTip", "MacRichTooltipShortcut", "(Command + Option)");
 			FText WinShortcut = NSLOCTEXT("SToolTip", "WinRichTooltipShortcut", "(Ctrl + Alt)");
 
 			FText KeyboardShortcut;
