Remove tutorial references from Documentation module.

diff --git a/Engine/Source/Editor/Documentation/Private/UDNParser.cpp b/Engine/Source/Editor/Documentation/Private/UDNParser.cpp
--- a/Engine/Source/Editor/Documentation/Private/UDNParser.cpp
+++ b/Engine/Source/Editor/Documentation/Private/UDNParser.cpp
@@ -36,7 +36,6 @@ namespace LinkPrefixes
 {
 
 static const FString DocLinkSpecifier( TEXT( "DOCLINK:" ) );
-static const FString TutorialLinkSpecifier( TEXT( "TUTORIALLINK:" ) );
 static const FString HttpLinkSpecifier( TEXT( "http://" ) );
 static const FString HttpsLinkSpecifier( TEXT( "https://" ) );
 
@@ -962,11 +961,6 @@ TSharedRef< SWidget > FUDNParser::GenerateExcerptContent( const FString& InLink,
 					const FString Link = Line.AdditionalContent[1].RightChop(LinkPrefixes::CodeLinkSpecifier.Len());
 					Excerpt.RichText += FString::Printf(TEXT("<a id=\"code\" href=\"%s\" style=\"%s\">%s</>"), *Link, *Style.HyperlinkStyleName.ToString(), *Line.AdditionalContent[0]);
 				}
-				else if(Line.AdditionalContent[1].Contains(LinkPrefixes::TutorialLinkSpecifier))
-				{
-					const FString Link = Line.AdditionalContent[1].RightChop(LinkPrefixes::TutorialLinkSpecifier.Len());
-					Excerpt.RichText += FString::Printf(TEXT("<a id=\"tutorial\" href=\"%s\" style=\"%s\">%s</>"), *Link, *Style.HyperlinkStyleName.ToString(), *Line.AdditionalContent[0]);
-				}
 				else
 				{
 					Excerpt.RichText += FString::Printf(TEXT("<a id=\"browser\" href=\"%s\" style=\"%s\">%s</>"), *InLink, *Style.HyperlinkStyleName.ToString(), *Line.AdditionalContent[0]);
@@ -1225,7 +1219,6 @@ void FUDNParser::HandleHyperlinkNavigate( FString AdditionalContent )
 void FUDNParser::NavigateToLink( FString AdditionalContent )
 {
 	static const FString DocLinkSpecifier( TEXT( "DOCLINK:" ) );
-	static const FString TutorialLinkSpecifier( TEXT( "TUTORIALLINK:" ) );
 	static const FString HttpLinkSPecifier( TEXT( "http://" ) );
 	static const FString HttpsLinkSPecifier( TEXT( "https://" ) );
 
@@ -1238,12 +1231,6 @@ void FUDNParser::NavigateToLink( FString AdditionalContent )
 		FString DocLink = AdditionalContent.RightChop(DocLinkSpecifier.Len());
 		IDocumentation::Get()->Open(DocLink, FDocumentationSourceInfo(TEXT("udn_parser")));
 	}
-	else if ( AdditionalContent.StartsWith( TutorialLinkSpecifier ) )
-	{
-		// internal link
-		FString InternalLink = AdditionalContent.RightChop( TutorialLinkSpecifier.Len() );
-		Configuration->OnNavigate.ExecuteIfBound( InternalLink );
-	}
 	else if ( AdditionalContent.StartsWith( HttpLinkSPecifier ) || AdditionalContent.StartsWith( HttpsLinkSPecifier ) )
 	{
 		// external link
