Decouple the documentation part of the "SDK not installed" handler from the tutorial system.

diff --git a/Engine/Plugins/Editor/GuidedTutorials/Source/IntroTutorials/Private/IntroTutorials.cpp b/Engine/Plugins/Editor/GuidedTutorials/Source/IntroTutorials/Private/IntroTutorials.cpp
--- a/Engine/Plugins/Editor/GuidedTutorials/Source/IntroTutorials/Private/IntroTutorials.cpp
+++ b/Engine/Plugins/Editor/GuidedTutorials/Source/IntroTutorials/Private/IntroTutorials.cpp
@@ -391,10 +391,6 @@ void FIntroTutorials::HandleSDKNotInstalled(const FString& PlatformName, const F
 			LaunchTutorialByName(InTutorialAsset);
 		}
 	}
-	else
-	{
-		IDocumentation::Get()->Open(InTutorialAsset);
-	}
 }
 
 void FIntroTutorials::HandleRequestResource(const FString& Category, const FString& InTutorialAsset)
diff --git a/Engine/Source/Editor/Documentation/Documentation.Build.cs b/Engine/Source/Editor/Documentation/Documentation.Build.cs
--- a/Engine/Source/Editor/Documentation/Documentation.Build.cs
+++ b/Engine/Source/Editor/Documentation/Documentation.Build.cs
@@ -22,6 +22,7 @@ namespace UnrealBuildTool.Rules
 				new string[]
 				{
 					"Core",
+					"MainFrame"
 					// ... add other public dependencies that you statically link with here ...
 				}
 			);
diff --git a/Engine/Source/Editor/Documentation/Private/Documentation.cpp b/Engine/Source/Editor/Documentation/Private/Documentation.cpp
--- a/Engine/Source/Editor/Documentation/Private/Documentation.cpp
+++ b/Engine/Source/Editor/Documentation/Private/Documentation.cpp
@@ -14,6 +14,7 @@
 #include "DocumentationLink.h"
 #include "SDocumentationToolTip.h"
 #include "Interfaces/IAnalyticsProvider.h"
+#include "Interfaces/IMainFrameModule.h"
 #include "EngineAnalytics.h"
 
 #define LOCTEXT_NAMESPACE "DocumentationActor"
@@ -25,7 +26,8 @@ TSharedRef< IDocumentation > FDocumentation::Create()
 
 FDocumentation::FDocumentation() 
 {
-
+	IMainFrameModule& MainFrameModule = FModuleManager::LoadModuleChecked<IMainFrameModule>(TEXT("MainFrame"));
+	MainFrameModule.OnMainFrameSDKNotInstalled().AddRaw(this, &FDocumentation::HandleSDKNotInstalled);
 }
 
 FDocumentation::~FDocumentation() 
@@ -267,4 +269,14 @@ TSharedRef< class SToolTip > FDocumentation::CreateToolTip(const TAttribute<FTex
 		];
 }
 
+void FDocumentation::HandleSDKNotInstalled(const FString& PlatformName, const FString& InDocumentationPage)
+{
+	if (FPackageName::IsValidLongPackageName(InDocumentationPage, true))
+	{
+		return;
+	}
+	IDocumentation::Get()->Open(InDocumentationPage);
+}
+
+
 #undef LOCTEXT_NAMESPACE
diff --git a/Engine/Source/Editor/Documentation/Private/Documentation.h b/Engine/Source/Editor/Documentation/Private/Documentation.h
--- a/Engine/Source/Editor/Documentation/Private/Documentation.h
+++ b/Engine/Source/Editor/Documentation/Private/Documentation.h
@@ -47,4 +47,6 @@ private:
 private:
 
 	TMap< FString, TWeakPtr< IDocumentationPage > > LoadedPages;
+
+	void HandleSDKNotInstalled(const FString& PlatformName, const FString& InDocumentationPage);
 };
