Fixed bug where Doc link generation has extra '/' when the sub path is empty.

diff --git a/Engine/Source/Editor/Documentation/Private/DocumentationLink.h b/Engine/Source/Editor/Documentation/Private/DocumentationLink.h
--- a/Engine/Source/Editor/Documentation/Private/DocumentationLink.h
+++ b/Engine/Source/Editor/Documentation/Private/DocumentationLink.h
@@ -51,12 +51,13 @@ public:
 		FString QueryString;
 		SplitLink(Link, Path, QueryString, Anchor);
 		const FString PartialPath = FString::Printf(TEXT("%s/index.html"), *Path);
-		
+		const FString TrimedPartialPath = PartialPath.TrimChar('/');
+
 		AddSourceInfoToQueryString(QueryString, Source);
 		FString URLRoot = GetUrlRoot();
 		FInternationalization& I18N = FInternationalization::Get();
 		FUnrealEdMisc::Get().ReplaceDocumentationURLWildcards(URLRoot, I18N.GetCurrentCulture());
-		return URLRoot + PartialPath + QueryString + Anchor;
+		return URLRoot + TrimedPartialPath + QueryString + Anchor;
 	}
 
 	static FString ToFilePath( const FString& Link )
