Support multiple tooltip sources and multiple base URLs for documentation.
- When the editor checks for enhanced tooltip content, it now looks in the Documentation/Source folder within the project folder and all enabled plugin folders.
- If the enhanced tooltip content contains a variable named ToolTipOverride, the content of that variable is used to override the default content for the initial simple tooltip.
- Documentation links can now use custom base URLs instead of always having the page ID appended to DocumentationURL in BaseEditor.ini.
- New base URLs can be defined in BaseEditor.ini or in DefaultEditor.ini within the project folder, as follows:
[/Script/Documentation.DocumentationSettings]
+DocumentationBaseUrls=(Id="ID",Url="http://myBaseUrl/{PAGEID}")
- They can also be added programmatically by calling IDocumentation::RegisterBaseUrl(), so plugins can add their own.
- The base URL can be specified anywhere a doc link is used. For links in code, the IDocumentation methods now accept an optional base URL parameter. For links taken from enhanced tooltip files, the ID of the base URL can now be specified at the file level in the BaseUrl metadata, or at the excerpt level in the [VAR:BaseUrl] variable.
- We now use {PAGEID} as a placeholder for where the doc link should go within the base URL.

diff --git a/Engine/Config/BaseEditor.ini b/Engine/Config/BaseEditor.ini
--- a/Engine/Config/BaseEditor.ini
+++ b/Engine/Config/BaseEditor.ini
@@ -268,7 +268,7 @@ RetainJpegFormat=True
 bShouldPromptOnCurveRemoveAll=true
 
 [UnrealEd.URLs]
-DocumentationURL="https://docs.unrealengine.com/{VERSION}/{I18N}/"
+DocumentationURL="https://docs.unrealengine.com/{VERSION}/{I18N}/{PAGEID}"
 APIDocsURL="https://docs.unrealengine.com/{VERSION}/{I18N}/API/"
 CommunityHomeURL="https://dev.epicgames.com/community"
 OnlineLearningURL="https://dev.epicgames.com/community/learning"
@@ -281,7 +281,7 @@ SourceCodeIDEURL_Mac="https://developer.apple.com/xcode"
 SourceCodeIDEURL_Linux="https://code.visualstudio.com/docs/setup/linux"
 SourceCodeIDEURL_Other="http://en.wikipedia.org/wiki/Integrated_development_environment"
 EpicGamesURL="http://UnrealEngine.com"
-ViewportControlsURL="https://docs.unrealengine.com/{VERSION}/{I18N}/BuildingWorlds/LevelEditor/Viewports/ViewportControls/"
+ViewportControlsURL="viewport-controls-in-unreal-engine"
 MeshSimplificationPluginsURL="/ue/marketplace/content-cat/assets/codeplugins"
 ReportABugURL="https://epicsupport.force.com/unrealengine/s/"
 IssueTrackerURL="https://issues.unrealengine.com/"
diff --git a/Engine/Plugins/Editor/ContentBrowser/ContentBrowserAssetDataSource/Source/ContentBrowserAssetDataSource/Private/AssetFileContextMenu.cpp b/Engine/Plugins/Editor/ContentBrowser/ContentBrowserAssetDataSource/Source/ContentBrowserAssetDataSource/Private/AssetFileContextMenu.cpp
--- a/Engine/Plugins/Editor/ContentBrowser/ContentBrowserAssetDataSource/Source/ContentBrowserAssetDataSource/Private/AssetFileContextMenu.cpp
+++ b/Engine/Plugins/Editor/ContentBrowser/ContentBrowserAssetDataSource/Source/ContentBrowserAssetDataSource/Private/AssetFileContextMenu.cpp
@@ -1890,7 +1890,8 @@ void FAssetFileContextMenu::ExecuteGoToDocsForAsset(UClass* SelectedClass, const
 		FString DocumentationLink = FEditorClassUtils::GetDocumentationLink(SelectedClass, ExcerptSection);
 		if (!DocumentationLink.IsEmpty())
 		{
-			IDocumentation::Get()->Open(DocumentationLink, FDocumentationSourceInfo(TEXT("cb_docs")));
+			FString DocumentationLinkBaseUrl = FEditorClassUtils::GetDocumentationLinkBaseUrl(SelectedClass);
+			IDocumentation::Get()->Open(DocumentationLink, FDocumentationSourceInfo(TEXT("cb_docs")), DocumentationLinkBaseUrl);
 		}
 	}
 }
diff --git a/Engine/Source/Editor/Documentation/Documentation.Build.cs b/Engine/Source/Editor/Documentation/Documentation.Build.cs
--- a/Engine/Source/Editor/Documentation/Documentation.Build.cs
+++ b/Engine/Source/Editor/Documentation/Documentation.Build.cs
@@ -36,7 +36,8 @@ namespace UnrealBuildTool.Rules
                     "InputCore",
                     "Slate",
 					"SlateCore",
-                    
+					"DeveloperSettings",
+					"Projects",
 					"EditorFramework",
                     "UnrealEd",
 					"Analytics",
diff --git a/Engine/Source/Editor/Documentation/Private/Documentation.cpp b/Engine/Source/Editor/Documentation/Private/Documentation.cpp
--- a/Engine/Source/Editor/Documentation/Private/Documentation.cpp
+++ b/Engine/Source/Editor/Documentation/Private/Documentation.cpp
@@ -12,13 +12,17 @@
 #include "UDNParser.h"
 #include "DocumentationPage.h"
 #include "DocumentationLink.h"
+#include "DocumentationSettings.h"
 #include "SDocumentationToolTip.h"
 #include "Interfaces/IAnalyticsProvider.h"
 #include "Interfaces/IMainFrameModule.h"
 #include "EngineAnalytics.h"
+#include "Interfaces/IPluginManager.h"
 
 #define LOCTEXT_NAMESPACE "DocumentationActor"
 
+DEFINE_LOG_CATEGORY(LogDocumentation);
+
 TSharedRef< IDocumentation > FDocumentation::Create() 
 {
 	return MakeShareable( new FDocumentation() );
@@ -28,6 +32,24 @@ FDocumentation::FDocumentation()
 {
 	IMainFrameModule& MainFrameModule = FModuleManager::LoadModuleChecked<IMainFrameModule>(TEXT("MainFrame"));
 	MainFrameModule.OnMainFrameSDKNotInstalled().AddRaw(this, &FDocumentation::HandleSDKNotInstalled);
+
+	if (const UDocumentationSettings* DocSettings = GetDefault<UDocumentationSettings>())
+	{
+		for (const FDocumentationBaseUrl& BaseUrl : DocSettings->DocumentationBaseUrls)
+		{
+			if (!RegisterBaseUrl(BaseUrl.Id, BaseUrl.Url))
+			{
+				UE_LOG(LogDocumentation, Warning, TEXT("Could not register documentation base URL: %s"), *BaseUrl.Id);
+			}
+		}
+	}
+
+	AddSourcePath(FPaths::Combine(FPaths::ProjectDir(), TEXT("Documentation"), TEXT("Source")));
+	for (TSharedRef<IPlugin> Plugin : IPluginManager::Get().GetEnabledPlugins())
+	{
+		AddSourcePath(FPaths::Combine(Plugin->GetBaseDir(), TEXT("Documentation"), TEXT("Source")));
+	}
+	AddSourcePath(FPaths::Combine(FPaths::EngineDir(), TEXT("Documentation"), TEXT("Source")));
 }
 
 FDocumentation::~FDocumentation() 
@@ -35,14 +57,14 @@ FDocumentation::~FDocumentation()
 
 }
 
-bool FDocumentation::OpenHome(FDocumentationSourceInfo Source) const
+bool FDocumentation::OpenHome(FDocumentationSourceInfo Source, const FString& BaseUrlId) const
 {
-	return Open( TEXT("%ROOT%"), Source );
+	return Open(TEXT("%ROOT%"), Source, BaseUrlId);
 }
 
-bool FDocumentation::OpenHome(const FCultureRef& Culture, FDocumentationSourceInfo Source) const
+bool FDocumentation::OpenHome(const FCultureRef& Culture, FDocumentationSourceInfo Source, const FString& BaseUrlId) const
 {
-	return Open(TEXT("%ROOT%"), Culture, Source);
+	return Open(TEXT("%ROOT%"), Culture, Source, BaseUrlId);
 }
 
 bool FDocumentation::OpenAPIHome(FDocumentationSourceInfo Source) const
@@ -60,7 +82,7 @@ bool FDocumentation::OpenAPIHome(FDocumentationSourceInfo Source) const
 	return false;
 }
 
-bool FDocumentation::Open(const FString& Link, FDocumentationSourceInfo Source) const
+bool FDocumentation::Open(const FString& Link, FDocumentationSourceInfo Source, const FString& BaseUrlId) const
 {
 	FString DocumentationUrl;
 
@@ -99,7 +121,7 @@ bool FDocumentation::Open(const FString& Link, FDocumentationSourceInfo Source)
 	{
 		// When opening a doc website we always request the most ideal culture for our documentation.
 		// The DNS will redirect us if necessary.
-		DocumentationUrl = FDocumentationLink::ToUrl(Link, Source);
+		DocumentationUrl = FDocumentationLink::ToUrl(Link, Source, BaseUrlId);
 	}
 
 	if (!DocumentationUrl.IsEmpty())
@@ -115,7 +137,7 @@ bool FDocumentation::Open(const FString& Link, FDocumentationSourceInfo Source)
 	return !DocumentationUrl.IsEmpty();
 }
 
-bool FDocumentation::Open(const FString& Link, const FCultureRef& Culture, FDocumentationSourceInfo Source) const
+bool FDocumentation::Open(const FString& Link, const FCultureRef& Culture, FDocumentationSourceInfo Source, const FString& BaseUrlId) const
 {
 	FString DocumentationUrl;
 
@@ -130,7 +152,7 @@ bool FDocumentation::Open(const FString& Link, const FCultureRef& Culture, FDocu
 
 	if (DocumentationUrl.IsEmpty())
 	{
-		DocumentationUrl = FDocumentationLink::ToUrl(Link, Culture, Source);
+		DocumentationUrl = FDocumentationLink::ToUrl(Link, Culture, Source, BaseUrlId);
 	}
 
 	if (!DocumentationUrl.IsEmpty())
@@ -146,12 +168,13 @@ bool FDocumentation::Open(const FString& Link, const FCultureRef& Culture, FDocu
 	return !DocumentationUrl.IsEmpty();
 }
 
-TSharedRef< SWidget > FDocumentation::CreateAnchor( const TAttribute<FString>& Link, const FString& PreviewLink, const FString& PreviewExcerptName ) const
+TSharedRef< SWidget > FDocumentation::CreateAnchor( const TAttribute<FString>& Link, const FString& PreviewLink, const FString& PreviewExcerptName, const TAttribute<FString>& BaseUrlId) const
 {
 	return SNew( SDocumentationAnchor )
 		.Link(Link)
 		.PreviewLink(PreviewLink)
-		.PreviewExcerptName(PreviewExcerptName);
+		.PreviewExcerptName(PreviewExcerptName)
+		.BaseUrlId(BaseUrlId);
 }
 
 TSharedRef< IDocumentationPage > FDocumentation::GetPage( const FString& Link, const TSharedPtr< FParserConfiguration >& Config, const FDocumentationStyle& Style )
@@ -185,8 +208,15 @@ bool FDocumentation::PageExists(const FString& Link) const
 		return true;
 	}
 
-	const FString SourcePath = FDocumentationLink::ToSourcePath(Link);
-	return FPaths::FileExists(SourcePath);
+	for (const FString& SourcePath : SourcePaths)
+	{
+		FString LinkPath = FDocumentationLink::ToSourcePath(Link, SourcePath);
+		if (FPaths::FileExists(LinkPath))
+		{
+			return true;
+		}
+	}
+	return false;
 }
 
 bool FDocumentation::PageExists(const FString& Link, const FCultureRef& Culture) const
@@ -197,8 +227,20 @@ bool FDocumentation::PageExists(const FString& Link, const FCultureRef& Culture)
 		return true;
 	}
 
-	const FString SourcePath = FDocumentationLink::ToSourcePath(Link, Culture);
-	return FPaths::FileExists(SourcePath);
+	for (const FString& SourcePath : SourcePaths)
+	{
+		FString LinkPath = FDocumentationLink::ToSourcePath(Link, Culture, SourcePath);
+		if (FPaths::FileExists(LinkPath))
+		{
+			return true;
+		}
+	}
+	return false;
+}
+
+const TArray <FString>& FDocumentation::GetSourcePaths() const
+{
+	return SourcePaths;
 }
 
 TSharedRef< class SToolTip > FDocumentation::CreateToolTip(const TAttribute<FText>& Text, const TSharedPtr<SWidget>& OverrideContent, const FString& Link, const FString& ExcerptName) const
@@ -278,5 +320,46 @@ void FDocumentation::HandleSDKNotInstalled(const FString& PlatformName, const FS
 	IDocumentation::Get()->Open(InDocumentationPage);
 }
 
+bool FDocumentation::RegisterBaseUrl(const FString& Id, const FString& Url)
+{
+	if (!Id.IsEmpty() && !Url.IsEmpty())
+	{
+		if (!RegisteredBaseUrls.Contains(Id))
+		{
+			RegisteredBaseUrls.Add(Id, Url);
+			return true;
+		}
+		UE_LOG(LogDocumentation, Warning, TEXT("Could not register documentation base URL with ID: %s. This ID is already in use."), *Id);
+		return false;
+	}
+	return false;
+}
+
+FString FDocumentation::GetBaseUrl(const FString& Id) const
+{
+	if (!Id.IsEmpty())
+	{
+		const FString* BaseUrl = RegisteredBaseUrls.Find(Id);
+		if (BaseUrl != NULL && !BaseUrl->IsEmpty())
+		{
+			return *BaseUrl;
+		}
+		UE_LOG(LogDocumentation, Warning, TEXT("Could not resolve base URL with ID: %s. It may not have been registered."), *Id);
+	}
+
+	FString DefaultUrl;
+	FUnrealEdMisc::Get().GetURL(TEXT("DocumentationURL"), DefaultUrl, true);
+	return DefaultUrl;
+}
+
+bool FDocumentation::AddSourcePath(const FString& Path)
+{
+	if (!Path.IsEmpty() && FPaths::DirectoryExists(Path))
+	{
+		SourcePaths.Add(Path);
+		return true;
+	}
+	return false;
+}
 
 #undef LOCTEXT_NAMESPACE
diff --git a/Engine/Source/Editor/Documentation/Private/Documentation.h b/Engine/Source/Editor/Documentation/Private/Documentation.h
--- a/Engine/Source/Editor/Documentation/Private/Documentation.h
+++ b/Engine/Source/Editor/Documentation/Private/Documentation.h
@@ -6,6 +6,8 @@
 #include "IDocumentationPage.h"
 #include "IDocumentation.h"
 
+DECLARE_LOG_CATEGORY_EXTERN(LogDocumentation, Log, All);
+
 template< typename ObjectType > class TAttribute;
 
 class FDocumentation : public IDocumentation
@@ -18,17 +20,17 @@ public:
 
 	virtual ~FDocumentation();
 
-	virtual bool OpenHome(FDocumentationSourceInfo Source = FDocumentationSourceInfo()) const override;
+	virtual bool OpenHome(FDocumentationSourceInfo Source = FDocumentationSourceInfo(), const FString& BaseUrlId = FString()) const override;
 
-	virtual bool OpenHome(const FCultureRef& Culture, FDocumentationSourceInfo Source = FDocumentationSourceInfo()) const override;
+	virtual bool OpenHome(const FCultureRef& Culture, FDocumentationSourceInfo Source = FDocumentationSourceInfo(), const FString& BaseUrlId = FString()) const override;
 
 	virtual bool OpenAPIHome(FDocumentationSourceInfo Source = FDocumentationSourceInfo()) const override;
 
-	virtual bool Open( const FString& Link, FDocumentationSourceInfo Source = FDocumentationSourceInfo() ) const override;
+	virtual bool Open( const FString& Link, FDocumentationSourceInfo Source = FDocumentationSourceInfo(), const FString& BaseUrlId = FString()) const override;
 
-	virtual bool Open( const FString& Link, const FCultureRef& Culture, FDocumentationSourceInfo Source = FDocumentationSourceInfo() ) const override;
+	virtual bool Open( const FString& Link, const FCultureRef& Culture, FDocumentationSourceInfo Source = FDocumentationSourceInfo(), const FString& BaseUrlId = FString()) const override;
 
-	virtual TSharedRef< SWidget > CreateAnchor( const TAttribute<FString>& Link, const FString& PreviewLink = FString(), const FString& PreviewExcerptName = FString() ) const override;
+	virtual TSharedRef< SWidget > CreateAnchor( const TAttribute<FString>& Link, const FString& PreviewLink = FString(), const FString& PreviewExcerptName = FString(), const TAttribute<FString>& BaseUrlId = FString()) const override;
 
 	virtual TSharedRef< IDocumentationPage > GetPage( const FString& Link, const TSharedPtr< FParserConfiguration >& Config, const FDocumentationStyle& Style = FDocumentationStyle() ) override;
 
@@ -36,10 +38,16 @@ public:
 
 	virtual bool PageExists(const FString& Link, const FCultureRef& Culture) const override;
 
+	virtual const TArray < FString >& GetSourcePaths() const override;
+
 	virtual TSharedRef< class SToolTip > CreateToolTip( const TAttribute<FText>& Text, const TSharedPtr<SWidget>& OverrideContent, const FString& Link, const FString& ExcerptName ) const override;
 	
 	virtual TSharedRef< class SToolTip > CreateToolTip(const TAttribute<FText>& Text, const TSharedRef<SWidget>& OverrideContent, const TSharedPtr<SVerticalBox>& DocVerticalBox, const FString& Link, const FString& ExcerptName) const override;
 
+	virtual bool RegisterBaseUrl(const FString& Id, const FString& Url) override;
+
+	virtual FString GetBaseUrl(const FString& Id) const override;
+
 private:
 
 	FDocumentation();
@@ -48,5 +56,11 @@ private:
 
 	TMap< FString, TWeakPtr< IDocumentationPage > > LoadedPages;
 
+	TMap< const FString, const FString > RegisteredBaseUrls;
+
+	TArray < FString > SourcePaths;
+
 	void HandleSDKNotInstalled(const FString& PlatformName, const FString& InDocumentationPage);
+
+	bool AddSourcePath(const FString& Path);
 };
diff --git a/Engine/Source/Editor/Documentation/Private/DocumentationLink.h b/Engine/Source/Editor/Documentation/Private/DocumentationLink.h
--- a/Engine/Source/Editor/Documentation/Private/DocumentationLink.h
+++ b/Engine/Source/Editor/Documentation/Private/DocumentationLink.h
@@ -12,17 +12,9 @@ class FDocumentationLink
 {
 public: 
 
-	static FString GetUrlRoot()
+	static FString GetUrlRoot(const FString& BaseUrlId = FString())
 	{
-		FString Url;
-		FUnrealEdMisc::Get().GetURL( TEXT("DocumentationURL"), Url, true );
-
-		if ( !Url.EndsWith( TEXT( "/" ) ) )
-		{
-			Url += TEXT( "/" );
-		}
-
-		return Url;
+		return IDocumentation::Get()->GetBaseUrl(BaseUrlId);
 	}
 
 	static FString GetHomeUrl()
@@ -39,25 +31,36 @@ public:
 		return Url;
 	}
 
-	static FString ToUrl(const FString& Link, FDocumentationSourceInfo const& Source)
+	static FString ToUrl(const FString& Link, FDocumentationSourceInfo const& Source, const FString& BaseUrlId = FString())
 	{
-		return ToUrl(Link, FInternationalization::Get().GetCurrentCulture(), Source);
+		return ToUrl(Link, FInternationalization::Get().GetCurrentCulture(), Source, BaseUrlId);
 	}
 
-	static FString ToUrl(const FString& Link, const FCultureRef& Culture, FDocumentationSourceInfo const& Source)
+	static FString ToUrl(const FString& Link, const FCultureRef& Culture, FDocumentationSourceInfo const& Source, const FString& BaseUrlId = FString())
 	{
-		FString Path;
-		FString Anchor;
+		// Get the base URL for this doc request, if any.
+		// It may have any or all of the following pieces in it.
+		FString UrlRoot = GetUrlRoot(BaseUrlId);
+		FString RootPath;
 		FString QueryString;
-		SplitLink(Link, Path, QueryString, Anchor);
-		const FString PartialPath = FString::Printf(TEXT("%s/index.html"), *Path);
-		const FString TrimedPartialPath = PartialPath.TrimChar('/');
-
+		FString Anchor;
+		SplitLink(UrlRoot, RootPath, QueryString, Anchor);
+
+		// Break up the incoming documentation link into pieces.
+		// Depending on where it comes from, the link may be just a page ID, or it may have any or all
+		// of the following pieces in it. If there's a query string, SplitLink will merge it with the existing query
+		// string found in the base URL (if any). If there's an anchor, SplitLink will use it /instead/ of the
+		// one in the BaseUrl (if any).
+		FString LinkPath;
+		SplitLink(Link, LinkPath, QueryString, Anchor);
+		const FString TrimmedPartialPath = LinkPath.TrimChar('/');
+
+		// Add query parameters for the source, if needed.
 		AddSourceInfoToQueryString(QueryString, Source);
-		FString URLRoot = GetUrlRoot();
-		FInternationalization& I18N = FInternationalization::Get();
-		FUnrealEdMisc::Get().ReplaceDocumentationURLWildcards(URLRoot, I18N.GetCurrentCulture());
-		return URLRoot + TrimedPartialPath + QueryString + Anchor;
+
+		FString AssembledUrl = RootPath + QueryString + Anchor;
+		FUnrealEdMisc::Get().ReplaceDocumentationURLWildcards(AssembledUrl, Culture, TrimmedPartialPath);
+		return AssembledUrl;
 	}
 
 	static FString ToFilePath( const FString& Link )
@@ -129,18 +132,18 @@ public:
 		return FString::Printf(TEXT("file:///%s%s%s"), *ToFilePath(Link, Culture), *QueryString, *Anchor);
 	}
 
-	static FString ToSourcePath(const FString& Link)
+	static FString ToSourcePath(const FString& Link, const FString& BasePath = FString())
 	{
 		FInternationalization& I18N = FInternationalization::Get();
 
-		FString SourcePath = ToSourcePath(Link, I18N.GetCurrentCulture());
+		FString SourcePath = ToSourcePath(Link, I18N.GetCurrentCulture(), BasePath);
 
 		if (!FPaths::FileExists(SourcePath))
 		{
 			const FCulturePtr FallbackCulture = I18N.GetCulture(TEXT("en"));
 			if (FallbackCulture.IsValid())
 			{
-				const FString FallbackSourcePath = ToSourcePath(Link, FallbackCulture.ToSharedRef());
+				const FString FallbackSourcePath = ToSourcePath(Link, FallbackCulture.ToSharedRef(), BasePath);
 				if (FPaths::FileExists(FallbackSourcePath))
 				{
 					SourcePath = FallbackSourcePath;
@@ -151,28 +154,33 @@ public:
 		return SourcePath;
 	}
 
-	static FString ToSourcePath(const FString& Link, const FCultureRef& Culture)
+	static FString ToSourcePath(const FString& Link, const FCultureRef& Culture, const FString& BasePath = FString())
 	{
 		FString Path;
 		FString Anchor;
 		FString QueryString;
 		SplitLink(Link, Path, QueryString, Anchor);
 
-		const FString FullDirectoryPath = FPaths::EngineDir() + TEXT( "Documentation/Source" ) + Path + "/";
+		FString DocSourcePath = BasePath;
+		if (DocSourcePath.IsEmpty() && !IDocumentation::Get()->GetSourcePaths().IsEmpty())
+		{
+			DocSourcePath = IDocumentation::Get()->GetSourcePaths().Last();
+		}
+		const FString FullDirectoryPath = FPaths::Combine(DocSourcePath, Path);
 
-		const FString WildCard = FString::Printf(TEXT("%s*.%s.udn"), *FullDirectoryPath, *(Culture->GetUnrealLegacyThreeLetterISOLanguageName()));
+		const FString WildCard = FString::Printf(TEXT("%s/*.%s.udn"), *FullDirectoryPath, *(Culture->GetUnrealLegacyThreeLetterISOLanguageName()));
 
 		TArray<FString> Filenames;
 		IFileManager::Get().FindFiles(Filenames, *WildCard, true, false);
 
 		if (Filenames.Num() > 0)
 		{
-			return FullDirectoryPath + Filenames[0];
+			return FPaths::Combine(FullDirectoryPath, Filenames[0]);
 		}
 
 		// Since the source file doesn't exist already make up a valid name for a new one
 		FString Category = FPaths::GetBaseFilename(Link);
-		return FString::Printf(TEXT("%s%s.%s.udn"), *FullDirectoryPath, *Category, *(Culture->GetUnrealLegacyThreeLetterISOLanguageName()));
+		return FString::Printf(TEXT("%s/%s.%s.udn"), *FullDirectoryPath, *Category, *(Culture->GetUnrealLegacyThreeLetterISOLanguageName()));
 	}
 
 private:
@@ -191,6 +199,13 @@ private:
 		}
 	}
 	
+	// Splits the Link parameter into three parts:
+	// - The path, which is everything up to the first ? or # character (or to the end of the link).
+	// - The query string, which is everything from the first ? to the first # (or to the end of the link).
+	// - The anchor, which is everything from the first # to the end of the link.
+	// If the query parameter contains anything already, whatever is found in the link will
+	// be appended to the existing value.
+	//
 	static void SplitLink( const FString& Link, /*OUT*/ FString& Path, /*OUT*/ FString& QueryString, /*OUT*/ FString& Anchor )
 	{
 		FString CleanedLink = Link;
@@ -199,12 +214,11 @@ private:
 		if ( CleanedLink == TEXT("%ROOT%") )
 		{
 			Path.Empty();
-			Anchor.Empty();
-			QueryString.Empty();
 		}
 		else
 		{
 			FString PathAndQueryString;
+			FString InAnchor = FString(Anchor).TrimChar('#');
 			if ( !CleanedLink.Split( TEXT("#"), &PathAndQueryString, &Anchor ) )
 			{
 				PathAndQueryString = CleanedLink;
@@ -214,6 +228,10 @@ private:
 				// ensure leading #
 				Anchor = FString( TEXT("#") ) + Anchor;
 			}
+			else if (!InAnchor.IsEmpty())
+			{
+				Anchor = FString(TEXT("#")) + InAnchor;
+			}
 
 			if ( Anchor.EndsWith( TEXT("/"), ESearchCase::CaseSensitive ) )
 			{
@@ -225,12 +243,8 @@ private:
 				PathAndQueryString.LeftInline(PathAndQueryString.Len() - 1, false);
 			}
 
-			if ( !PathAndQueryString.IsEmpty() && !PathAndQueryString.StartsWith( TEXT("/"), ESearchCase::CaseSensitive ) )
-			{
-				PathAndQueryString = FString(TEXT("/")) + PathAndQueryString;
-			}
-
 			// split path and query string
+			FString InQuery = FString(QueryString).TrimChar('?');
 			if (!PathAndQueryString.Split(TEXT("?"), &Path, &QueryString))
 			{
 				Path = PathAndQueryString;
@@ -238,7 +252,7 @@ private:
 			else if (!QueryString.IsEmpty())
 			{
 				// ensure leading ?
-				QueryString = FString(TEXT("?")) + QueryString;
+				QueryString = FString(TEXT("?")) + InQuery + QueryString;
 			}
 		}
 	}
diff --git a/Engine/Source/Editor/Documentation/Private/SDocumentationAnchor.cpp b/Engine/Source/Editor/Documentation/Private/SDocumentationAnchor.cpp
--- a/Engine/Source/Editor/Documentation/Private/SDocumentationAnchor.cpp
+++ b/Engine/Source/Editor/Documentation/Private/SDocumentationAnchor.cpp
@@ -5,6 +5,7 @@
 #include "EditorClassUtils.h"
 #include "Styling/AppStyle.h"
 #include "IDocumentation.h"
+#include "IDocumentationPage.h"
 #include "Widgets/Images/SImage.h"
 #include "Widgets/Input/SButton.h"
 #include "Widgets/Layout/SBox.h"
@@ -14,6 +15,7 @@
 void SDocumentationAnchor::Construct(const FArguments& InArgs )
 {
 	Link = InArgs._Link;
+	BaseUrlId = InArgs._BaseUrlId;
 
 	SetVisibility(TAttribute<EVisibility>::CreateLambda([this]()
 		{
@@ -27,21 +29,41 @@ void SDocumentationAnchor::Construct(const FArguments& InArgs )
 	}
 
 	const FString PreviewLink = InArgs._PreviewLink;
-	// All in-editor UDN documents must live under the Shared/ folder
-	ensure(PreviewLink.IsEmpty() || PreviewLink.StartsWith(TEXT("Shared/")));
+	const FString PreviewExcerptName = InArgs._PreviewExcerptName;
+	if (!PreviewLink.IsEmpty() && !PreviewExcerptName.IsEmpty() && BaseUrlId.Get().IsEmpty())
+	{
+		TSharedRef<IDocumentation> Documentation = IDocumentation::Get();
+		if (Documentation->PageExists(PreviewLink))
+		{
+			TSharedPtr<IDocumentationPage> DocumentationPage = Documentation->GetPage(PreviewLink, NULL);
+
+			FExcerpt Excerpt;
+			if (DocumentationPage->HasExcerpt(PreviewExcerptName))
+			{
+				if (DocumentationPage->GetExcerpt(PreviewExcerptName, Excerpt))
+				{
+					if (FString* BaseUrlValue = Excerpt.Variables.Find(TEXT("BaseUrl")))
+					{
+						BaseUrlId = *BaseUrlValue;
+					}
+				}
+			}
+		}
+
+	}
 
 	ChildSlot
 	[
 		SAssignNew(Button, SSimpleButton)
 		.OnClicked(this, &SDocumentationAnchor::OnClicked)
 		.Icon(FAppStyle::Get().GetBrush("Icons.Help"))
-		.ToolTip(IDocumentation::Get()->CreateToolTip(ToolTipText, nullptr, PreviewLink, InArgs._PreviewExcerptName))
+		.ToolTip(IDocumentation::Get()->CreateToolTip(ToolTipText, nullptr, PreviewLink, PreviewExcerptName))
 	];
 }
 
 
 FReply SDocumentationAnchor::OnClicked() const
 {
-	IDocumentation::Get()->Open(Link.Get(FString()), FDocumentationSourceInfo(TEXT("doc_anchors")));
+	IDocumentation::Get()->Open(Link.Get(FString()), FDocumentationSourceInfo(TEXT("doc_anchors")), BaseUrlId.Get(FString()));
 	return FReply::Handled();
 }
diff --git a/Engine/Source/Editor/Documentation/Private/SDocumentationAnchor.h b/Engine/Source/Editor/Documentation/Private/SDocumentationAnchor.h
--- a/Engine/Source/Editor/Documentation/Private/SDocumentationAnchor.h
+++ b/Engine/Source/Editor/Documentation/Private/SDocumentationAnchor.h
@@ -17,6 +17,8 @@ public:
 
 		/** The string for the link to follow when clicked  */
 		SLATE_ATTRIBUTE( FString, Link )
+		/** The base URL for the Link, if any is needed  */
+		SLATE_ATTRIBUTE(FString, BaseUrlId)
 	SLATE_END_ARGS()
 
 	void Construct(const FArguments& InArgs);
@@ -28,6 +30,7 @@ private:
 private:
 
 	TAttribute<FString> Link;
+	TAttribute<FString> BaseUrlId;
 	TSharedPtr<class SButton> Button;
 	TSharedPtr<class SImage> ButtonImage;
 };
diff --git a/Engine/Source/Editor/Documentation/Private/SDocumentationToolTip.cpp b/Engine/Source/Editor/Documentation/Private/SDocumentationToolTip.cpp
--- a/Engine/Source/Editor/Documentation/Private/SDocumentationToolTip.cpp
+++ b/Engine/Source/Editor/Documentation/Private/SDocumentationToolTip.cpp
@@ -37,12 +37,6 @@ void SDocumentationToolTip::Construct( const FArguments& InArgs )
 	bAddDocumentation = InArgs._AddDocumentation;
 	DocumentationMargin = InArgs._DocumentationMargin;
 
-	if ( !DocumentationLink.IsEmpty() )
-	{
-		// All in-editor udn documents must live under the Shared/ folder
-		ensure( InArgs._DocumentationLink.StartsWith( TEXT("Shared/") ) );
-	}
-
 	ExcerptName = InArgs._ExcerptName;
 	IsShowingFullTip = false;
 
@@ -66,6 +60,30 @@ void SDocumentationToolTip::Construct( const FArguments& InArgs )
 
 void SDocumentationToolTip::ConstructSimpleTipContent()
 {
+	// If there a UDN file that matches the DocumentationLink path, and that page has an excerpt whose name
+	// matches ExcerptName, and that excerpt has a variable named ToolTipOverride, use the content of that
+	// variable instead of the default TextContent.
+	if (!DocumentationLink.IsEmpty() && !ExcerptName.IsEmpty())
+	{
+		TSharedRef<IDocumentation> Documentation = IDocumentation::Get();
+		if (Documentation->PageExists(DocumentationLink))
+		{
+			DocumentationPage = Documentation->GetPage(DocumentationLink, NULL);
+
+			FExcerpt Excerpt;
+			if (DocumentationPage->HasExcerpt(ExcerptName))
+			{
+				if (DocumentationPage->GetExcerpt(ExcerptName, Excerpt))
+				{
+					if (FString* TooltipValue = Excerpt.Variables.Find(TEXT("ToolTipOverride")))
+					{
+						TextContent = FText::FromString(*TooltipValue);
+					}
+				}
+			}
+		}
+	}
+
 	TSharedPtr< SVerticalBox > VerticalBox;
 	if ( !OverrideContent.IsValid() )
 	{
@@ -295,19 +313,14 @@ void SDocumentationToolTip::ConstructFullTipContent()
 				];
 
 			FString* FullDocumentationLink = Excerpts[ ExcerptIndex ].Variables.Find( TEXT("ToolTipFullLink") );
+			FString* ExcerptBaseUrl = Excerpts[ExcerptIndex].Variables.Find(TEXT("BaseUrl"));
 			if ( FullDocumentationLink != NULL && !FullDocumentationLink->IsEmpty() )
 			{
-				struct Local
+				FString BaseUrl = FString();
+				if (ExcerptBaseUrl != NULL)
 				{
-					static void OpenLink( FString Link )
-					{
-						if (!IDocumentation::Get()->Open(Link, FDocumentationSourceInfo(TEXT("rich_tooltips"))))
-						{
-							FNotificationInfo Info( NSLOCTEXT("SToolTip", "FailedToOpenLink", "Failed to Open Link") );
-							FSlateNotificationManager::Get().AddNotification(Info);
-						}
-					}
-				};
+					BaseUrl = *ExcerptBaseUrl;
+				}
 
 				Box->AddSlot()
 				.HAlign( HAlign_Center )
@@ -317,21 +330,18 @@ void SDocumentationToolTip::ConstructFullTipContent()
 						.Text( NSLOCTEXT( "SToolTip", "GoToFullDocsLinkMessage", "see full documentation" ) )
 						.TextStyle( &HyperlinkTextStyleInfo )
 						.UnderlineStyle( &HyperlinkButtonStyleInfo )
-						.OnNavigate_Static( &Local::OpenLink, *FullDocumentationLink )
+						.OnNavigate_Static([](FString Link, FString BaseUrl) {
+								if (!IDocumentation::Get()->Open(Link, FDocumentationSourceInfo(TEXT("rich_tooltips")), BaseUrl))
+								{
+									FNotificationInfo Info(NSLOCTEXT("SToolTip", "FailedToOpenLink", "Failed to Open Link"));
+									FSlateNotificationManager::Get().AddNotification(Info);
+								}
+							}, *FullDocumentationLink, BaseUrl)
 				];
 			}
 
 			if ( GetDefault<UEditorPerProjectUserSettings>()->bDisplayDocumentationLink && FSlateApplication::Get().SupportsSourceAccess() )
 			{
-				struct Local
-				{
-					static void EditSource( FString Link, int32 LineNumber )
-					{
-						ISourceCodeAccessModule& SourceCodeAccessModule = FModuleManager::LoadModuleChecked<ISourceCodeAccessModule>("SourceCodeAccess");
-						SourceCodeAccessModule.GetAccessor().OpenFileAtLine(Link, LineNumber);
-					}
-				};
-
 				Box->AddSlot()
 				.AutoHeight()
 				.HAlign( HAlign_Center )
@@ -340,7 +350,11 @@ void SDocumentationToolTip::ConstructFullTipContent()
 						.Text( NSLOCTEXT( "SToolTip", "EditDocumentationMessage_Edit", "edit" ) )
 						.TextStyle( &HyperlinkTextStyleInfo )
 						.UnderlineStyle( &HyperlinkButtonStyleInfo )
-						.OnNavigate_Static(&Local::EditSource, FPaths::ConvertRelativePathToFull(FDocumentationLink::ToSourcePath(DocumentationLink, FInternationalization::Get().GetCurrentCulture())), Excerpts[ExcerptIndex].LineNumber)
+						// todo: needs to update to point to the "real" source file used for the excerpt
+						.OnNavigate_Static([](FString Link, int32 LineNumber) {
+								ISourceCodeAccessModule& SourceCodeAccessModule = FModuleManager::LoadModuleChecked<ISourceCodeAccessModule>("SourceCodeAccess");
+								SourceCodeAccessModule.GetAccessor().OpenFileAtLine(Link, LineNumber);
+							}, FPaths::ConvertRelativePathToFull(FDocumentationLink::ToSourcePath(DocumentationLink, FInternationalization::Get().GetCurrentCulture())), Excerpts[ExcerptIndex].LineNumber)
 				];
 			}
 		}
diff --git a/Engine/Source/Editor/Documentation/Private/UDNParser.cpp b/Engine/Source/Editor/Documentation/Private/UDNParser.cpp
--- a/Engine/Source/Editor/Documentation/Private/UDNParser.cpp
+++ b/Engine/Source/Editor/Documentation/Private/UDNParser.cpp
@@ -16,7 +16,7 @@
 #include "Editor/EditorPerProjectUserSettings.h"
 #include "Developer/MessageLog/Public/MessageLogModule.h"
 #include "Logging/MessageLog.h"
-#include "Editor/Documentation/Private/DocumentationLink.h"
+#include "DocumentationLink.h"
 #include "ISourceCodeAccessor.h"
 #include "ISourceCodeAccessModule.h"
 #include "IContentBrowserSingleton.h"
@@ -101,6 +101,7 @@ void FUDNParser::Initialize()
 	TokenLibrary.Add(FTokenPair(TEXT("Title:"), EUDNToken::MetadataTitle));
 	TokenLibrary.Add(FTokenPair(TEXT("Crumbs:"), EUDNToken::MetadataCrumbs));
 	TokenLibrary.Add(FTokenPair(TEXT("Description:"), EUDNToken::MetadataDescription));
+	TokenLibrary.Add(FTokenPair(TEXT("BaseUrl:"), EUDNToken::MetadataBaseUrl));
 	TokenLibrary.Add(FTokenPair(TEXT("%"), EUDNToken::Percentage));
 	TokenLibrary.Add(FTokenPair(TEXT("*"), EUDNToken::Asterisk));
 
@@ -210,6 +211,10 @@ void FUDNParser::Initialize()
 	TokenArray.Add(EUDNToken::MetadataDescription);
 	LineLibrary.Add(FTokenConfiguration(TokenArray, FUDNLine::MetadataDescription, true));
 
+	TokenArray.Empty();
+	TokenArray.Add(EUDNToken::MetadataBaseUrl);
+	LineLibrary.Add(FTokenConfiguration(TokenArray, FUDNLine::MetadataBaseUrl, true));
+
 	TokenArray.Empty();
 	TokenArray.Add(EUDNToken::OpenBracket);
 	TokenArray.Add(EUDNToken::Variable);
@@ -250,16 +255,14 @@ FUDNParser::~FUDNParser()
 bool FUDNParser::LoadLink( const FString& Link, TArray<FString>& ContentLines )
 {
 	FMessageLog UDNParserLog(UDNParseErrorLog);
-
-	const FString SourcePath = FDocumentationLink::ToSourcePath( Link );
 	
-	if (!FPaths::FileExists(SourcePath))
+	if (!FPaths::FileExists(Link))
 	{
 		return false;
 	}
 
 	TArray<uint8> Buffer;
-	bool bLoadSuccess = FFileHelper::LoadFileToArray(Buffer, *SourcePath);
+	bool bLoadSuccess = FFileHelper::LoadFileToArray(Buffer, *Link);
 	if ( bLoadSuccess )
 	{
 		FString Result;
@@ -306,7 +309,7 @@ bool FUDNParser::LoadLink( const FString& Link, TArray<FString>& ContentLines )
 	}
 	else
 	{
-		UDNParserLog.Error(FText::Format(LOCTEXT("LoadingError", "Loading document '{0}' failed."), FText::FromString(SourcePath)));
+		UDNParserLog.Error(FText::Format(LOCTEXT("LoadingError", "Loading document '{0}' failed."), FText::FromString(Link)));
 	}
 
 	if ( !bLoadSuccess && GetDefault<UEditorPerProjectUserSettings>()->bDisplayDocumentationLink )
@@ -321,41 +324,49 @@ bool FUDNParser::Parse(const FString& Link, TArray<FExcerpt>& OutExcerpts, FUDNP
 {
 	FMessageLog UDNParserLog(UDNParseErrorLog);
 
-	TArray<FString> ContentLines;
-	if ( LoadLink( Link, ContentLines ) )
+	// Call LoadLink for all content folders, and amalgamate the excerpts from all matching UDN files into a single list.
+	const TArray<FString> SourcePaths = IDocumentation::Get()->GetSourcePaths();
+	for (const FString& SourcePath : SourcePaths)
 	{
-		TArray<FExcerpt> TempExcerpts;
-		const FString SourcePath = FDocumentationLink::ToSourcePath( Link );
-		bool bParseSuccess = ParseSymbols( Link, ContentLines, FPaths::GetPath( SourcePath ), TempExcerpts, OutMetadata );
-
-		if (bParseSuccess)
-		{
-			OutExcerpts = TempExcerpts;
-			return true;
-		}
-		else
+		FString DocFilePath = FDocumentationLink::ToSourcePath(Link, SourcePath);
+		TArray<FString> ContentLines;
+		if (LoadLink(DocFilePath, ContentLines))
 		{
-			if ( GetDefault<UEditorPerProjectUserSettings>()->bDisplayDocumentationLink )
+			TArray<FExcerpt> TempExcerpts;
+			bool bParseSuccess = ParseSymbols(DocFilePath, ContentLines, FPaths::GetPath(DocFilePath), TempExcerpts, OutMetadata);
+
+			if (bParseSuccess)
 			{
-				UDNParserLog.Open();
+				OutExcerpts.Append(TempExcerpts);
 			}
+			else
+			{
+				if (GetDefault<UEditorPerProjectUserSettings>()->bDisplayDocumentationLink)
+				{
+					UDNParserLog.Open();
+				}
 
-			UDNParserLog.Error(FText::Format(LOCTEXT("GeneralParsingError", "Parsing document '{0}' failed."), FText::FromString(SourcePath)));
+				UDNParserLog.Error(FText::Format(LOCTEXT("GeneralParsingError", "Parsing document '{0}' failed."), FText::FromString(DocFilePath)));
+			}
 		}
 	}
-
-	return false;
+	return true;
 }
 
 bool FUDNParser::GetExcerptContent( const FString& Link, FExcerpt& Excerpt )
 {
 	FMessageLog UDNParserLog(UDNParseErrorLog);
 
+	FString SourcePath = FDocumentationLink::ToSourcePath(Link);
+	if (!Excerpt.SourcePath.IsEmpty())
+	{
+		SourcePath = Excerpt.SourcePath;
+	}
 	TArray<FString> ContentLines;
 
-	if ( LoadLink( Link, ContentLines ) )
+	if ( LoadLink(SourcePath, ContentLines ) )
 	{
-		Excerpt.Content = GenerateExcerptContent( Link, Excerpt, ContentLines, Excerpt.LineNumber );
+		Excerpt.Content = GenerateExcerptContent(SourcePath, Excerpt, ContentLines, Excerpt.LineNumber );
 		return true;
 	}
 	else
@@ -365,7 +376,7 @@ bool FUDNParser::GetExcerptContent( const FString& Link, FExcerpt& Excerpt )
 			UDNParserLog.Open();
 		}
 
-		UDNParserLog.Error(FText::Format(LOCTEXT("GeneralExcerptError", "Generating a Widget for document '{0}' Excerpt '{1}' failed."), FText::FromString( FDocumentationLink::ToSourcePath( Link ) ), FText::FromString(Excerpt.Name) ));
+		UDNParserLog.Error(FText::Format(LOCTEXT("GeneralExcerptError", "Generating a Widget for document '{0}' Excerpt '{1}' failed."), FText::FromString(SourcePath), FText::FromString(Excerpt.Name) ));
 	}
 
 	return false;
@@ -750,8 +761,7 @@ TSharedRef< SWidget > FUDNParser::GenerateExcerptContent( const FString& InLink,
 {
 	FMessageLog UDNParserLog(UDNParseErrorLog);
 
-	const FString SourcePath = FDocumentationLink::ToSourcePath( InLink );
-	const FString FullPath = FPaths::GetPath( SourcePath );
+	const FString FullPath = FPaths::GetPath(InLink);
 
 	FSlateFontInfo Header1Font = FCoreStyle::GetDefaultFontStyle("Regular", 18);
 	FSlateFontInfo Header2Font = FCoreStyle::GetDefaultFontStyle("Regular", 14);
@@ -1087,7 +1097,11 @@ bool FUDNParser::ParseSymbols(const FString& Link, const TArray<FString>& Conten
 			
 			if (ExcerptStack.Num() == 0)
 			{
-				OutExcerpts.Add(FExcerpt(ExcerptName, NULL, Variables, ExcerptStartingLineNumber));
+				if (!OutMetadata.BaseUrl.IsEmpty() && !Variables.Contains("BaseUrl"))
+				{
+					Variables.Add("BaseUrl", OutMetadata.BaseUrl);
+				}
+				OutExcerpts.Add(FExcerpt(ExcerptName, NULL, Variables, ExcerptStartingLineNumber, Link));
 				OutMetadata.ExcerptNames.Add( ExcerptName );
 				Variables.Empty();
 				ExcerptStartingLineNumber = 0;
@@ -1123,7 +1137,7 @@ bool FUDNParser::ParseSymbols(const FString& Link, const TArray<FString>& Conten
 				break;
 			}
 
-			Variables.Add( VariableName, VariableValue );
+			Variables.Add( VariableName, VariableValue.TrimStartAndEnd());
 
 			VariableName.Empty();
 			VariableValue.Empty();
@@ -1149,7 +1163,7 @@ bool FUDNParser::ParseSymbols(const FString& Link, const TArray<FString>& Conten
 				break;
 			}
 
-			Variables.Add( VariableName, VariableValue );
+			Variables.Add( VariableName, VariableValue.TrimStartAndEnd());
 
 			VariableName.Empty();
 			VariableValue.Empty();
@@ -1163,6 +1177,7 @@ bool FUDNParser::ParseSymbols(const FString& Link, const TArray<FString>& Conten
 			case FUDNLine::MetadataTitle: OutMetadata.Title = FText::FromString(Line.AdditionalContent[0]); break;
 			case FUDNLine::MetadataCrumbs: OutMetadata.Crumbs = FText::FromString(Line.AdditionalContent[0]); break;
 			case FUDNLine::MetadataDescription: OutMetadata.Description = FText::FromString(Line.AdditionalContent[0]); break;
+			case FUDNLine::MetadataBaseUrl: OutMetadata.BaseUrl = Line.AdditionalContent[0]; break;
 			}
 		}
 		else
@@ -1179,10 +1194,24 @@ bool FUDNParser::ParseSymbols(const FString& Link, const TArray<FString>& Conten
 				{
 					if ( !VariableName.IsEmpty() )
 					{
+						if (!VariableValue.IsEmpty())
+						{
+							VariableValue += '\n';
+						}
 						VariableValue += Line.AdditionalContent[0];
 					}
 				}
 				break;
+			case FUDNLine::Whitespace:
+				{
+					if (!VariableName.IsEmpty())
+					{
+						if (!VariableValue.IsEmpty())
+						{
+							VariableValue += '\n';
+						}
+					}
+				}
 			}
 		}
 	}
diff --git a/Engine/Source/Editor/Documentation/Private/UDNParser.h b/Engine/Source/Editor/Documentation/Private/UDNParser.h
--- a/Engine/Source/Editor/Documentation/Private/UDNParser.h
+++ b/Engine/Source/Editor/Documentation/Private/UDNParser.h
@@ -18,7 +18,8 @@ struct FUDNPageMetadata
 		, Title()
 		, Crumbs()
 		, Description()
-		, ExcerptNames() 
+		, ExcerptNames()
+		, BaseUrl()
 	{}
 
 	FString Availability;
@@ -26,6 +27,7 @@ struct FUDNPageMetadata
 	FText Crumbs;
 	FText Description;
 	TSet< FString > ExcerptNames;
+	FString BaseUrl;
 };
 
 /** Represents a single UDN Markdown token */
@@ -50,6 +52,7 @@ namespace EUDNToken
 		MetadataTitle,
 		MetadataCrumbs,
 		MetadataDescription,
+		MetadataBaseUrl,
 		Percentage,
 		Asterisk
 	};
@@ -90,6 +93,7 @@ public:
 		MetadataTitle,
 		MetadataCrumbs,
 		MetadataDescription,
+		MetadataBaseUrl,
 		Variable,
 		VariableOpen,
 		VariableClose,
diff --git a/Engine/Source/Editor/Documentation/Public/DocumentationSettings.h b/Engine/Source/Editor/Documentation/Public/DocumentationSettings.h
new file mode 100644
--- /dev/null
+++ b/Engine/Source/Editor/Documentation/Public/DocumentationSettings.h
@@ -0,0 +1,44 @@
+#pragma once
+
+#include "Engine/DeveloperSettings.h"
+#include "IDocumentation.h"
+
+#include "DocumentationSettings.generated.h"
+
+
+USTRUCT()
+struct FDocumentationBaseUrl
+{
+	GENERATED_USTRUCT_BODY()
+
+public:
+	UPROPERTY()
+		FString Id;
+
+	UPROPERTY()
+		FString Url;
+
+	FDocumentationBaseUrl() {};
+
+	FDocumentationBaseUrl(FString const& InId, FString const& InUrl)
+		: Id(InId), Url(InUrl)
+	{};
+
+	/** Returns true if either the ID or URL is unset. */
+	bool IsEmpty() const
+	{
+		return Id.IsEmpty() || Url.IsEmpty();
+	}
+};
+
+
+UCLASS(config=Editor)
+class UDocumentationSettings : public UDeveloperSettings
+{
+	GENERATED_BODY()
+
+public:
+	// Array of base URLs for documentation links that are loaded on startup
+	UPROPERTY(Config)
+	TArray<FDocumentationBaseUrl> DocumentationBaseUrls;
+};
\ No newline at end of file
diff --git a/Engine/Source/Editor/Documentation/Public/IDocumentation.h b/Engine/Source/Editor/Documentation/Public/IDocumentation.h
--- a/Engine/Source/Editor/Documentation/Public/IDocumentation.h
+++ b/Engine/Source/Editor/Documentation/Public/IDocumentation.h
@@ -177,17 +177,17 @@ public:
 
 public:
 
-	virtual bool OpenHome(FDocumentationSourceInfo Source = FDocumentationSourceInfo()) const = 0;
+	virtual bool OpenHome(FDocumentationSourceInfo Source = FDocumentationSourceInfo(), const FString& BaseUrlId = FString()) const = 0;
 
-	virtual bool OpenHome(const FCultureRef& Culture, FDocumentationSourceInfo Source = FDocumentationSourceInfo()) const = 0;
+	virtual bool OpenHome(const FCultureRef& Culture, FDocumentationSourceInfo Source = FDocumentationSourceInfo(), const FString & BaseUrlId = FString()) const = 0;
 
 	virtual bool OpenAPIHome(FDocumentationSourceInfo Source = FDocumentationSourceInfo()) const = 0;
 
-	virtual bool Open(const FString& Link, FDocumentationSourceInfo Source = FDocumentationSourceInfo()) const = 0;
+	virtual bool Open(const FString& Link, FDocumentationSourceInfo Source = FDocumentationSourceInfo(), const FString& BaseUrlId = FString()) const = 0;
 
-	virtual bool Open(const FString& Link, const FCultureRef& Culture, FDocumentationSourceInfo Source = FDocumentationSourceInfo()) const = 0;
+	virtual bool Open(const FString& Link, const FCultureRef& Culture, FDocumentationSourceInfo Source = FDocumentationSourceInfo(), const FString& BaseUrlId = FString()) const = 0;
 
-	virtual TSharedRef< class SWidget > CreateAnchor( const TAttribute<FString>& Link, const FString& PreviewLink = FString(), const FString& PreviewExcerptName = FString() ) const = 0;
+	virtual TSharedRef< class SWidget > CreateAnchor( const TAttribute<FString>& Link, const FString& PreviewLink = FString(), const FString& PreviewExcerptName = FString(), const TAttribute<FString>& BaseUrlId = FString()) const = 0;
 
 	virtual TSharedRef< class IDocumentationPage > GetPage( const FString& Link, const TSharedPtr< FParserConfiguration >& Config, const FDocumentationStyle& Style = FDocumentationStyle() ) = 0;
 
@@ -195,7 +195,13 @@ public:
 
 	virtual bool PageExists(const FString& Link, const FCultureRef& Culture) const = 0;
 
+	virtual const TArray < FString >& GetSourcePaths() const = 0;
+
 	virtual TSharedRef< class SToolTip > CreateToolTip( const TAttribute<FText>& Text, const TSharedPtr<SWidget>& OverrideContent, const FString& Link, const FString& ExcerptName ) const = 0;
 
 	virtual TSharedRef< class SToolTip > CreateToolTip(const TAttribute<FText>& Text, const TSharedRef<SWidget>& OverrideContent, const TSharedPtr<class SVerticalBox>& DocVerticalBox, const FString& Link, const FString& ExcerptName) const = 0;
+
+	virtual bool RegisterBaseUrl(const FString& Id, const FString& Url) = 0;
+
+	virtual FString GetBaseUrl(const FString& Id) const = 0;
 };
diff --git a/Engine/Source/Editor/Documentation/Public/IDocumentationPage.h b/Engine/Source/Editor/Documentation/Public/IDocumentationPage.h
--- a/Engine/Source/Editor/Documentation/Public/IDocumentationPage.h
+++ b/Engine/Source/Editor/Documentation/Public/IDocumentationPage.h
@@ -21,14 +21,27 @@ struct FExcerpt
 		, LineNumber( InLineNumber )
 	{}
 
-	/** Excerpt name */
+	FExcerpt(const FString& InName, const TSharedPtr< SWidget >& InContent, const TMap< FString, FString >& InVariables, int32 InLineNumber, const FString& InSourcePath)
+		: Name(InName)
+		, Content(InContent)
+		, Variables(InVariables)
+		, LineNumber(InLineNumber)
+		, SourcePath(InSourcePath)
+	{}
+
+	/** The name of the excerpt, as specified in the [EXCERPT:NAME] line at the top of the excerpt. */
 	FString Name;
-	/** Slate content for excerpt */
+	/** The Slate content generated for the excerpt. */
 	TSharedPtr<SWidget> Content;
+	/** The content of the [VAR] sections read from the excerpt. The keys are the names from the [VAR:NAME] lines
+		that start the variables, and the value for each key is the text inside the [VAR:NAME].....[/VAR] block. */
 	TMap< FString, FString > Variables;
+	/** The line number that this excerpt begins within its source file. */
 	int32 LineNumber;
-	/** Rich text version of the excerpt */
+	/** A rich text version of the excerpt content. */
 	FString RichText;
+	/** The path to the file that this excerpt was read from. */
+	FString SourcePath;
 };
 
 class IDocumentationPage
diff --git a/Engine/Source/Editor/Kismet/Private/BlueprintEditor.cpp b/Engine/Source/Editor/Kismet/Private/BlueprintEditor.cpp
--- a/Engine/Source/Editor/Kismet/Private/BlueprintEditor.cpp
+++ b/Engine/Source/Editor/Kismet/Private/BlueprintEditor.cpp
@@ -7807,12 +7807,31 @@ FString FBlueprintEditor::GetDocLinkForSelectedNode()
 	return DocumentationLink;
 }
 
+FString FBlueprintEditor::GetDocLinkBaseUrlForSelectedNode()
+{
+	FString DocumentationLinkBaseUrl;
+
+	if (const UEdGraphNode* SelectedGraphNode = GetSingleSelectedNode())
+	{
+		const FString DocLink = SelectedGraphNode->GetDocumentationLink();
+		const FString DocExcerpt = SelectedGraphNode->GetDocumentationExcerptName();
+
+		if (!DocLink.IsEmpty() && !DocExcerpt.IsEmpty())
+		{
+			DocumentationLinkBaseUrl = FEditorClassUtils::GetDocumentationLinkBaseUrlFromExcerpt(DocLink, DocExcerpt);
+		}
+	}
+
+	return DocumentationLinkBaseUrl;
+}
+
 void FBlueprintEditor::OnGoToDocumentation()
 {
 	const FString DocumentationLink = GetDocLinkForSelectedNode();
+	const FString DocumentationLinkBaseUrl = GetDocLinkBaseUrlForSelectedNode();
 	if (!DocumentationLink.IsEmpty())
 	{
-		IDocumentation::Get()->Open(DocumentationLink, FDocumentationSourceInfo(TEXT("rightclick_bpnode")));
+		IDocumentation::Get()->Open(DocumentationLink, FDocumentationSourceInfo(TEXT("rightclick_bpnode")), DocumentationLinkBaseUrl);
 	}
 }
 
diff --git a/Engine/Source/Editor/Kismet/Public/BlueprintEditor.h b/Engine/Source/Editor/Kismet/Public/BlueprintEditor.h
--- a/Engine/Source/Editor/Kismet/Public/BlueprintEditor.h
+++ b/Engine/Source/Editor/Kismet/Public/BlueprintEditor.h
@@ -1285,6 +1285,9 @@ private:
 	/** Util to try and get doc link for the currently selected node */
 	FString GetDocLinkForSelectedNode();
 
+	/** Util to try and get the base URL for the doc link for the currently selected node */
+	FString GetDocLinkBaseUrlForSelectedNode();
+
 	/** Set the enabled state for currently-selected nodes */
 	void OnSetEnabledStateForSelectedNodes(ENodeEnabledState NewState);
 
diff --git a/Engine/Source/Editor/LevelEditor/Private/LevelEditorActions.cpp b/Engine/Source/Editor/LevelEditor/Private/LevelEditorActions.cpp
--- a/Engine/Source/Editor/LevelEditor/Private/LevelEditorActions.cpp
+++ b/Engine/Source/Editor/LevelEditor/Private/LevelEditorActions.cpp
@@ -190,9 +190,7 @@ void FLevelEditorActionCallbacks::BrowseViewportControls()
 	FString URL;
 	if (FUnrealEdMisc::Get().GetURL(TEXT("ViewportControlsURL"), URL))
 	{
-		// documentation URLS are now stored with placeholders for internationalization and version
-		FUnrealEdMisc::Get().ReplaceDocumentationURLWildcards(URL, FInternationalization::Get().GetCurrentCulture());
-		FPlatformProcess::LaunchURL(*URL, NULL, NULL);
+		IDocumentation::Get()->Open(URL, FDocumentationSourceInfo(TEXT("help_menu")));
 	}
 }
 
@@ -1231,7 +1229,8 @@ void FLevelEditorActionCallbacks::GoToDocsForActor_Clicked()
 		FString DocumentationLink = FEditorClassUtils::GetDocumentationLink(SelectedActorInfo.SelectionClass);
 		if (!DocumentationLink.IsEmpty())
 		{
-			IDocumentation::Get()->Open( DocumentationLink, FDocumentationSourceInfo(TEXT("rightclick_viewdoc")) );
+			FString DocumentationLinkBaseUrl = FEditorClassUtils::GetDocumentationLinkBaseUrl(SelectedActorInfo.SelectionClass);
+			IDocumentation::Get()->Open( DocumentationLink, FDocumentationSourceInfo(TEXT("rightclick_viewdoc")), DocumentationLinkBaseUrl);
 		}
 	}
 }
diff --git a/Engine/Source/Editor/MaterialEditor/Private/MaterialEditor.cpp b/Engine/Source/Editor/MaterialEditor/Private/MaterialEditor.cpp
--- a/Engine/Source/Editor/MaterialEditor/Private/MaterialEditor.cpp
+++ b/Engine/Source/Editor/MaterialEditor/Private/MaterialEditor.cpp
@@ -4725,28 +4725,57 @@ FString FMaterialEditor::GetDocLinkForSelectedNode()
 {
 	FString DocumentationLink;
 
-	TArray<UObject*> SelectedNodes = GetSelectedNodes().Array();
+	FGraphPanelSelectionSet SelectedNodes = GetSelectedNodes();
 	if (SelectedNodes.Num() == 1)
 	{
-		UMaterialGraphNode* SelectedGraphNode = Cast<UMaterialGraphNode>(SelectedNodes[0]);
-		if (SelectedGraphNode != NULL)
+		for (UObject* ObjectInSelection : SelectedNodes)
 		{
-			FString DocLink = SelectedGraphNode->GetDocumentationLink();
-			FString DocExcerpt = SelectedGraphNode->GetDocumentationExcerptName();
+			if (ObjectInSelection != NULL)
+			{
+				UMaterialGraphNode* SelectedGraphNode = Cast<UMaterialGraphNode>(ObjectInSelection);
+				FString DocLink = SelectedGraphNode->GetDocumentationLink();
+				FString DocExcerpt = SelectedGraphNode->GetDocumentationExcerptName();
 
-			DocumentationLink = FEditorClassUtils::GetDocumentationLinkFromExcerpt(DocLink, DocExcerpt);
+				DocumentationLink = FEditorClassUtils::GetDocumentationLinkFromExcerpt(DocLink, DocExcerpt);
+			}
+			break;
 		}
 	}
 
 	return DocumentationLink;
 }
 
+FString FMaterialEditor::GetDocLinkBaseUrlForSelectedNode()
+{
+	FString DocumentationLinkBaseUrl;
+
+	FGraphPanelSelectionSet SelectedNodes = GetSelectedNodes();
+	if (SelectedNodes.Num() == 1)
+	{
+		for (UObject* ObjectInSelection : SelectedNodes)
+		{
+			if (ObjectInSelection != NULL)
+			{
+				UMaterialGraphNode* SelectedGraphNode = Cast<UMaterialGraphNode>(ObjectInSelection);
+				FString DocLink = SelectedGraphNode->GetDocumentationLink();
+				FString DocExcerpt = SelectedGraphNode->GetDocumentationExcerptName();
+
+				DocumentationLinkBaseUrl = FEditorClassUtils::GetDocumentationLinkBaseUrlFromExcerpt(DocLink, DocExcerpt);
+			}
+			break;
+		}
+	}
+
+	return DocumentationLinkBaseUrl;
+}
+
 void FMaterialEditor::OnGoToDocumentation()
 {
 	FString DocumentationLink = GetDocLinkForSelectedNode();
+	FString DocumentationLinkBaseUrl = GetDocLinkBaseUrlForSelectedNode();
 	if (!DocumentationLink.IsEmpty())
 	{
-		IDocumentation::Get()->Open(DocumentationLink, FDocumentationSourceInfo(TEXT("rightclick_matnode")));
+		IDocumentation::Get()->Open(DocumentationLink, FDocumentationSourceInfo(TEXT("rightclick_matnode")), DocumentationLinkBaseUrl);
 	}
 }
 
diff --git a/Engine/Source/Editor/MaterialEditor/Private/MaterialEditor.h b/Engine/Source/Editor/MaterialEditor/Private/MaterialEditor.h
--- a/Engine/Source/Editor/MaterialEditor/Private/MaterialEditor.h
+++ b/Engine/Source/Editor/MaterialEditor/Private/MaterialEditor.h
@@ -793,6 +793,9 @@ private:
 	/** Util to try and get doc link for the currently selected node */
 	FString GetDocLinkForSelectedNode();
 
+	/** Util to try and get the base URL for the doc link for the currently selected node */
+	FString GetDocLinkBaseUrlForSelectedNode();
+
 	/** Callback from the Asset Registry when an asset is renamed. */
 	void RenameAssetFromRegistry(const FAssetData& InAddedAssetData, const FString& InNewName);
 
diff --git a/Engine/Source/Editor/UnrealEd/Private/EditorClassUtils.cpp b/Engine/Source/Editor/UnrealEd/Private/EditorClassUtils.cpp
--- a/Engine/Source/Editor/UnrealEd/Private/EditorClassUtils.cpp
+++ b/Engine/Source/Editor/UnrealEd/Private/EditorClassUtils.cpp
@@ -56,6 +56,27 @@ FString FEditorClassUtils::GetDocumentationLinkFromExcerpt(const FString& DocLin
 	return DocumentationLink;
 }
 
+FString FEditorClassUtils::GetDocumentationLinkBaseUrlFromExcerpt(const FString& DocLink, const FString DocExcerpt)
+{
+	FString DocumentationLinkBaseUrl;
+	TSharedRef<IDocumentation> Documentation = IDocumentation::Get();
+	if (Documentation->PageExists(DocLink))
+	{
+		TSharedRef<IDocumentationPage> ClassDocs = Documentation->GetPage(DocLink, NULL);
+
+		FExcerpt Excerpt;
+		if (ClassDocs->GetExcerpt(DocExcerpt, Excerpt))
+		{
+			FString* DocumentationLinkBaseUrlValue = Excerpt.Variables.Find(TEXT("BaseUrl"));
+			if (DocumentationLinkBaseUrlValue)
+			{
+				DocumentationLinkBaseUrl = *DocumentationLinkBaseUrlValue;
+			}
+		}
+	}
+
+	return DocumentationLinkBaseUrl;
+}
 
 FString FEditorClassUtils::GetDocumentationLink(const UClass* Class, const FString& OverrideExcerpt)
 {
@@ -65,15 +86,24 @@ FString FEditorClassUtils::GetDocumentationLink(const UClass* Class, const FStri
 	return GetDocumentationLinkFromExcerpt(ClassDocsPage, ExcerptSection);
 }
 
+FString FEditorClassUtils::GetDocumentationLinkBaseUrl(const UClass* Class, const FString& OverrideExcerpt)
+{
+	const FString ClassDocsPage = GetDocumentationPage(Class);
+	const FString ExcerptSection = (OverrideExcerpt.IsEmpty() ? GetDocumentationExcerpt(Class) : OverrideExcerpt);
+
+	return GetDocumentationLinkBaseUrlFromExcerpt(ClassDocsPage, ExcerptSection);
+}
+
 
 TSharedRef<SWidget> FEditorClassUtils::GetDocumentationLinkWidget(const UClass* Class)
 {
 	TSharedRef<SWidget> DocLinkWidget = SNullWidget::NullWidget;
 	const FString DocumentationLink = GetDocumentationLink(Class);
+	const FString DocumentationLinkBaseUrl = GetDocumentationLinkBaseUrl(Class);
 
 	if (!DocumentationLink.IsEmpty())
 	{
-		DocLinkWidget = IDocumentation::Get()->CreateAnchor(DocumentationLink);
+		DocLinkWidget = IDocumentation::Get()->CreateAnchor(DocumentationLink, FString(), FString(), DocumentationLinkBaseUrl);
 	}
 
 	return DocLinkWidget;
@@ -85,7 +115,11 @@ TSharedRef<SWidget> FEditorClassUtils::GetDynamicDocumentationLinkWidget(const T
 	{
 		return GetDocumentationLink(ClassAttribute.Get(nullptr));
 	};
-	return IDocumentation::Get()->CreateAnchor(TAttribute<FString>::CreateLambda(GetLink));
+	auto GetBaseUrl = [ClassAttribute]()
+	{
+		return GetDocumentationLinkBaseUrl(ClassAttribute.Get(nullptr));
+	};
+	return IDocumentation::Get()->CreateAnchor(TAttribute<FString>::CreateLambda(GetLink), FString(), FString(), TAttribute<FString>::CreateLambda(GetBaseUrl));
 }
 
 TSharedRef<SWidget> FEditorClassUtils::GetSourceLink(const UClass* Class, const FSourceLinkParams& Params)
diff --git a/Engine/Source/Editor/UnrealEd/Private/UnrealEdMisc.cpp b/Engine/Source/Editor/UnrealEd/Private/UnrealEdMisc.cpp
--- a/Engine/Source/Editor/UnrealEd/Private/UnrealEdMisc.cpp
+++ b/Engine/Source/Editor/UnrealEd/Private/UnrealEdMisc.cpp
@@ -1901,11 +1901,12 @@ bool FUnrealEdMisc::GetURL( const TCHAR* InKey, FString& OutURL, const bool bChe
 	return bFound;
 }
 
-void FUnrealEdMisc::ReplaceDocumentationURLWildcards(FString& Url, const FCultureRef& Culture)
+void FUnrealEdMisc::ReplaceDocumentationURLWildcards(FString& Url, const FCultureRef& Culture, const FString& PageId)
 {
 	static FString Version = FString::FromInt(FEngineVersion::Current().GetMajor()) + TEXT(".") + FString::FromInt(FEngineVersion::Current().GetMinor());
 	Url.ReplaceInline(TEXT("{VERSION}"), *Version);
 	Url.ReplaceInline(TEXT("{I18N}"), *(Culture->GetName()));
+	Url.ReplaceInline(TEXT("{PAGEID}"), *PageId);
 }
 
 FString FUnrealEdMisc::GetExecutableForCommandlets() const
diff --git a/Engine/Source/Editor/UnrealEd/Public/EditorClassUtils.h b/Engine/Source/Editor/UnrealEd/Public/EditorClassUtils.h
--- a/Engine/Source/Editor/UnrealEd/Public/EditorClassUtils.h
+++ b/Engine/Source/Editor/UnrealEd/Public/EditorClassUtils.h
@@ -59,6 +59,19 @@ namespace FEditorClassUtils
 	 */
 	UNREALED_API FString GetDocumentationLinkFromExcerpt(const FString& DocLink, const FString DocExcerpt);
 
+	/**
+	 * Returns the ID of the base documentation URL set for this class in its documentation excerpt.
+	 *
+	 * @param	Class		Class we want to build a link for
+	 * @return				The ID of the base URL
+	 */
+	UNREALED_API FString GetDocumentationLinkBaseUrl(const UClass* Class, const FString& OverrideExcerpt = FString());
+
+	/**
+	 * Returns the ID of the base documentation URL set in the specified excerpt.
+	 */
+	UNREALED_API FString GetDocumentationLinkBaseUrlFromExcerpt(const FString& DocLink, const FString DocExcerpt);
+
 	/**
 	 * Creates a link widget to the documentation for a given class
 	 *
diff --git a/Engine/Source/Editor/UnrealEd/Public/UnrealEdMisc.h b/Engine/Source/Editor/UnrealEd/Public/UnrealEdMisc.h
--- a/Engine/Source/Editor/UnrealEd/Public/UnrealEdMisc.h
+++ b/Engine/Source/Editor/UnrealEd/Public/UnrealEdMisc.h
@@ -212,7 +212,7 @@ public:
 	 */
 	bool GetURL( const TCHAR* InKey, FString& OutURL, const bool bCheckRocket = false ) const;
 
-	void ReplaceDocumentationURLWildcards(FString& Url, const FCultureRef& Culture);
+	void ReplaceDocumentationURLWildcards(FString& Url, const FCultureRef& Culture, const FString& PageId = FString());
 
 	/** Returns the editor executable to use to execute commandlets */
 	FString GetExecutableForCommandlets() const;
