Add missing copyright boilerplate.

diff --git a/Engine/Source/Editor/Documentation/Public/DocumentationSettings.h b/Engine/Source/Editor/Documentation/Public/DocumentationSettings.h
--- a/Engine/Source/Editor/Documentation/Public/DocumentationSettings.h
+++ b/Engine/Source/Editor/Documentation/Public/DocumentationSettings.h
@@ -1,3 +1,4 @@
+// Copyright Epic Games, Inc. All Rights Reserved.
 #pragma once
 
 #include "Engine/DeveloperSettings.h"
