Improved tooltips for pinning actors.
Adds new capabilities to the enhaced tooltip system to support this:
- Can now create enhanced tooltips for menu items created with `AddMenuEntry`, using tooltip localization IDs as excerpt names
- Bulleted lists and numbered lists now render with a hanging indent (using a horizontal box with two slots)
- Can share content between two tooltips using the new `[VAR:ExcerptAlias]` variable.

diff --git a/Engine/Documentation/Source/Shared/MenuEntries/SceneOutliner_ActorBrowsingMode/SceneOutliner_ActorBrowsingMode.INT.udn b/Engine/Documentation/Source/Shared/MenuEntries/SceneOutliner_ActorBrowsingMode/SceneOutliner_ActorBrowsingMode.INT.udn
new file mode 100644
--- /dev/null
+++ b/Engine/Documentation/Source/Shared/MenuEntries/SceneOutliner_ActorBrowsingMode/SceneOutliner_ActorBrowsingMode.INT.udn
@@ -0,0 +1,35 @@
+Availability:NoPublish
+Title:
+Crumbs:
+Description:
+
+[EXCERPT:PinTooltip]
+Pinning an actor keeps it loaded in the editor no matter what regions you have loaded in the World Partition panel.
+
+This can be practical if you need to temporarily work with specific actors that aren't in any loaded regions, or if there are specific actors you want to keep loaded all the time as you load and unload World Partition regions.
+
+---
+
+To pin or unpin an actor:
+
+- Toggle the icon in the Pin column of the Outliner, or right-click the actor in the Outliner and choose Pin or Unpin from the contextual menu.
+
+---
+
+Note:
+
+- Pinned actors will still be unloaded if they are in at least one Data Layer and none of those Data Layers are loaded.
+
+- Pinning an actor in the editor has no effect at runtime. If you need to exclude an actor from control by World Partition at runtime, disable the actor's Is Spatially Loaded setting in the Details panel.
+
+[VAR:ToolTipFullLink]
+world-partition-in-unreal-engine
+[/VAR]
+[/EXCERPT:PinTooltip]
+
+[EXCERPT:UnpinTooltip]
+[VAR:ExcerptAlias]
+PinTooltip
+[/VAR]
+[/EXCERPT:UnpinTooltip]
+
diff --git a/Engine/Source/Editor/Documentation/Private/DocumentationModule.cpp b/Engine/Source/Editor/Documentation/Private/DocumentationModule.cpp
--- a/Engine/Source/Editor/Documentation/Private/DocumentationModule.cpp
+++ b/Engine/Source/Editor/Documentation/Private/DocumentationModule.cpp
@@ -54,6 +54,18 @@ private:
 		if ( Action.IsValid() )
 		{
 			return Documentation->CreateToolTip( ToolTipText, OverrideContent, FString( TEXT("Shared/") ) + Action->GetBindingContext().ToString(), Action->GetCommandName().ToString() );
+		} 
+		else
+		{
+			if (!ToolTipText.Get().IsEmpty())
+			{
+				TOptional<FString> TooltipNS = FTextInspector::GetNamespace(ToolTipText.Get());
+				TOptional<FString> TooltipKey = FTextInspector::GetKey(ToolTipText.Get());
+				if (TooltipNS.IsSet() && TooltipKey.IsSet())
+				{
+					return Documentation->CreateToolTip(ToolTipText, OverrideContent, FString(TEXT("Shared/MenuEntries/")) + TooltipNS.GetValue(), TooltipKey.GetValue());
+				}
+			}
 		}
 
 		TSharedPtr< SWidget > ToolTipContent;
diff --git a/Engine/Source/Editor/Documentation/Private/DocumentationPage.cpp b/Engine/Source/Editor/Documentation/Private/DocumentationPage.cpp
--- a/Engine/Source/Editor/Documentation/Private/DocumentationPage.cpp
+++ b/Engine/Source/Editor/Documentation/Private/DocumentationPage.cpp
@@ -16,6 +16,13 @@ FDocumentationPage::~FDocumentationPage()
 
 bool FDocumentationPage::GetExcerptContent( FExcerpt& Excerpt )
 {
+	if (FString* AliasName = StoredMetadata.ExcerptAliases.Find(Excerpt.Name))
+	{
+		if (HasExcerpt(*AliasName))
+		{
+			Excerpt.Name = *AliasName;
+		}
+	}
 	for (int32 Index = 0; Index < StoredExcerpts.Num(); ++Index)
 	{
 		if ( Excerpt.Name == StoredExcerpts[ Index ].Name )
diff --git a/Engine/Source/Editor/Documentation/Private/UDNParser.cpp b/Engine/Source/Editor/Documentation/Private/UDNParser.cpp
--- a/Engine/Source/Editor/Documentation/Private/UDNParser.cpp
+++ b/Engine/Source/Editor/Documentation/Private/UDNParser.cpp
@@ -161,6 +161,10 @@ void FUDNParser::Initialize()
 	}
 	LineLibrary.Add(FTokenConfiguration(TokenArray, FUDNLine::HorizontalRule));
 
+	TokenArray.Empty();
+	TokenArray.Add(EUDNToken::Dash);
+	LineLibrary.Add(FTokenConfiguration(TokenArray, FUDNLine::BulletContent, true));
+
 	TokenArray.Empty();
 	TokenArray.Add(EUDNToken::Pound);
 	TokenArray.Add(EUDNToken::Pound);
@@ -785,6 +789,49 @@ void FUDNParser::AddContentToExcerpt(TSharedPtr<SVerticalBox> Box, const FString
 	}
 }
 
+void FUDNParser::AddListItemToExcerpt(TSharedPtr<SVerticalBox> Box, const FString& LeftContentSource, const FString& RightContentSource, FExcerpt& Excerpt)
+{
+	if (!LeftContentSource.IsEmpty() && !RightContentSource.IsEmpty())
+	{
+		Box->AddSlot()
+			.AutoHeight()
+			.HAlign(HAlign_Center)
+			[
+				SNew(SBox)
+				.HAlign(HAlign_Left)
+				.WidthOverride(ContentWidth)
+				.Padding(FMargin(0, 0, 0, 8.f))
+				[
+					SNew(SHorizontalBox)
+					+ SHorizontalBox::Slot()
+					.AutoWidth()
+					[
+						SNew(SBox)
+						.MinDesiredWidth(40.f)
+						.Padding(FMargin(16.f, 0, 0, 0))
+						.HAlign(HAlign_Left)
+						[
+							SNew(STextBlock)
+							.Text(FText::FromString(LeftContentSource))
+							.TextStyle(FAppStyle::Get(), Style.ContentStyleName)
+						]
+					]
+					+ SHorizontalBox::Slot()
+					.FillWidth(100)
+					[
+						SNew(STextBlock)
+						.Text(FText::FromString(RightContentSource))
+						.TextStyle(FAppStyle::Get(), Style.ContentStyleName)
+						.AutoWrapText(true)
+					]
+				]
+			];
+
+		AddLineSeperator(Excerpt);
+		Excerpt.RichText += FString::Printf(TEXT("<TextStyle Style=\"%s\">%s %s</>"), *Style.ContentStyleName.ToString(), *LeftContentSource, *RightContentSource);
+	}
+}
+
 TSharedRef< SWidget > FUDNParser::GenerateExcerptContent( const FString& InLink, FExcerpt& Excerpt, const TArray<FString>& ContentLines, int32 StartingLineIndex )
 {
 	FMessageLog UDNParserLog(UDNParseErrorLog);
@@ -919,13 +966,16 @@ TSharedRef< SWidget > FUDNParser::GenerateExcerptContent( const FString& InLink,
 				Excerpt.RichText += FString::Printf(TEXT("<TextStyle Style=\"%s\">%s</>"), *Style.BoldContentStyleName.ToString(), *Line.AdditionalContent[0]);
 				break;
 			case FUDNLine::NumberedContent:
-				AddContentToExcerpt(Box, CurrentStringContent, Excerpt);
-				CurrentStringContent = FString::Printf(TEXT("%i. %s"), CurrentNumbering, *Line.AdditionalContent[0]);
 				AddContentToExcerpt(Box, CurrentStringContent, Excerpt);
 				CurrentStringContent.Empty();
-
+				AddListItemToExcerpt(Box, FString::Printf(TEXT("%i."), CurrentNumbering), Line.AdditionalContent[0], Excerpt);
 				++CurrentNumbering;
 				break;
+			case FUDNLine::BulletContent:
+				AddContentToExcerpt(Box, CurrentStringContent, Excerpt);
+				CurrentStringContent.Empty();
+				AddListItemToExcerpt(Box, TEXT("•"), Line.AdditionalContent[0], Excerpt);
+				break;
 			case FUDNLine::HorizontalRule:
 				AddContentToExcerpt(Box, CurrentStringContent, Excerpt);
 				CurrentStringContent.Empty();
@@ -1129,6 +1179,10 @@ bool FUDNParser::ParseSymbols(const FString& Link, const TArray<FString>& Conten
 				{
 					Variables.Add("BaseUrl", OutMetadata.BaseUrl);
 				}
+				if (Variables.Contains("ExcerptAlias") && !Variables.Find("ExcerptAlias")->IsEmpty())
+				{
+					OutMetadata.ExcerptAliases.Add(ExcerptName, *Variables.Find("ExcerptAlias"));
+				}
 				OutExcerpts.Add(FExcerpt(ExcerptName, NULL, Variables, ExcerptStartingLineNumber, Link));
 				OutMetadata.ExcerptNames.Add( ExcerptName );
 				Variables.Empty();
@@ -1214,6 +1268,7 @@ bool FUDNParser::ParseSymbols(const FString& Link, const TArray<FString>& Conten
 			{
 			case FUDNLine::Content:
 			case FUDNLine::NumberedContent:
+			case FUDNLine::BulletContent:
 			case FUDNLine::Header1:
 			case FUDNLine::Header2:
 			case FUDNLine::Image:
diff --git a/Engine/Source/Editor/Documentation/Private/UDNParser.h b/Engine/Source/Editor/Documentation/Private/UDNParser.h
--- a/Engine/Source/Editor/Documentation/Private/UDNParser.h
+++ b/Engine/Source/Editor/Documentation/Private/UDNParser.h
@@ -27,6 +27,7 @@ struct FUDNPageMetadata
 		, Crumbs()
 		, Description()
 		, ExcerptNames()
+		, ExcerptAliases()
 		, BaseUrl()
 	{}
 
@@ -35,6 +36,7 @@ struct FUDNPageMetadata
 	FText Crumbs;
 	FText Description;
 	TSet< FString > ExcerptNames;
+	TMap< FString, FString > ExcerptAliases;
 	FString BaseUrl;
 };
 
@@ -105,7 +107,8 @@ public:
 		Variable,
 		VariableOpen,
 		VariableClose,
-		BoldContent
+		BoldContent,
+		BulletContent
 	};
 	Type ContentType;
 
@@ -161,6 +164,9 @@ private:
 	/** Adds the content text source to the scrollbox */
 	void AddContentToExcerpt(TSharedPtr<SVerticalBox> Box, const FString& ContentSource, FExcerpt& Excerpt);
 	
+	/** Adds the content text sources to the scrollbox in two columns, emulating a list item with hanging indent */
+	void AddListItemToExcerpt(TSharedPtr<SVerticalBox> Box, const FString& LeftContentSource, const FString& RightContentSource, FExcerpt& Excerpt);
+
 	/** Gets the dynamic brush for the given filename */
 	TSharedPtr<FSlateDynamicImageBrush> GetDynamicBrushFromImagePath(FString Filename);
 
diff --git a/Engine/Source/Editor/SceneOutliner/Private/ActorBrowsingMode.cpp b/Engine/Source/Editor/SceneOutliner/Private/ActorBrowsingMode.cpp
--- a/Engine/Source/Editor/SceneOutliner/Private/ActorBrowsingMode.cpp
+++ b/Engine/Source/Editor/SceneOutliner/Private/ActorBrowsingMode.cpp
@@ -531,7 +531,7 @@ void FActorBrowsingMode::RegisterContextMenu()
 								Section.AddMenuEntry(
 									"PinItems",
 									LOCTEXT("Pin", "Pin"),
-									FText(),
+									LOCTEXT("PinTooltip", "Keep the selected items loaded in the editor even when they don't overlap a loaded World Partition region"),
 									FSlateIcon(),
 									FUIAction(FExecuteAction::CreateSP(Context->SceneOutliner.Pin().Get(), &SSceneOutliner::PinSelectedItems)));
 							}
@@ -543,7 +543,7 @@ void FActorBrowsingMode::RegisterContextMenu()
 								Section.AddMenuEntry(
 									"UnpinItems",
 									LOCTEXT("Unpin", "Unpin"),
-									FText(),
+									LOCTEXT("UnpinTooltip", "Allow the World Partition system to load and unload the selected items automatically"),
 									FSlateIcon(),
 									FUIAction(FExecuteAction::CreateSP(Context->SceneOutliner.Pin().Get(), &SSceneOutliner::UnpinSelectedItems)));
 							}
diff --git a/Engine/Source/Editor/SceneOutliner/Private/SceneOutlinerPinnedColumn.cpp b/Engine/Source/Editor/SceneOutliner/Private/SceneOutlinerPinnedColumn.cpp
--- a/Engine/Source/Editor/SceneOutliner/Private/SceneOutlinerPinnedColumn.cpp
+++ b/Engine/Source/Editor/SceneOutliner/Private/SceneOutlinerPinnedColumn.cpp
@@ -3,6 +3,8 @@
 #include "SceneOutlinerPinnedColumn.h"
 #include "Widgets/Images/SImage.h"
 #include "Widgets/Views/STreeView.h"
+#include "IDocumentation.h"
+#include "Widgets/SToolTip.h"
 
 #define LOCTEXT_NAMESPACE "SceneOutlinerPinnedActorColumn"
 
@@ -186,7 +188,7 @@ SHeaderRow::FColumn::FArguments FSceneOutlinerPinnedColumn::ConstructHeaderRowCo
 		.VAlignHeader(VAlign_Center)
 		.HAlignCell(HAlign_Center)
 		.VAlignCell(VAlign_Center)
-		.DefaultTooltip(LOCTEXT("SceneOutlinerPinnedColumnToolTip", "Pinned - Force loaded in editor"))
+		.ToolTip(IDocumentation::Get()->CreateToolTip(LOCTEXT("SceneOutlinerPinnedColumnToolTip", "Pinned: always loaded in editor"), nullptr, "Shared/MenuEntries/SceneOutliner_ActorBrowsingMode", "PinTooltip"))
 		[
 			SNew(SImage)
 			.ColorAndOpacity(FSlateColor::UseForeground())
