Removing [EXCERPT:NAME] comment text to mitigate an API doc tool issue.

#rnx

diff --git a/Engine/Source/Editor/Documentation/Public/IDocumentationPage.h b/Engine/Source/Editor/Documentation/Public/IDocumentationPage.h
--- a/Engine/Source/Editor/Documentation/Public/IDocumentationPage.h
+++ b/Engine/Source/Editor/Documentation/Public/IDocumentationPage.h
@@ -29,7 +29,7 @@ struct FExcerpt
 		, SourcePath(InSourcePath)
 	{}
 
-	/** The name of the excerpt, as specified in the [EXCERPT:NAME] line at the top of the excerpt. */
+	/** The name of the excerpt, as specified in the line at the top of the excerpt. */
 	FString Name;
 	/** The Slate content generated for the excerpt. */
 	TSharedPtr<SWidget> Content;
