Updates to tooltips and to the enhanced tooltip system.
- Shortcuts on toolbar item tooltips now use a dedicated display style instead of being embedded in the text string.
- Enhanced tooltips now have an animated transition between states.
- Updated visual styles for enhanced tooltips.
- Images in enhanced tooltips now print their alt text as captions in italics below the images.

diff --git a/Engine/Source/Editor/Documentation/Private/Documentation.cpp b/Engine/Source/Editor/Documentation/Private/Documentation.cpp
--- a/Engine/Source/Editor/Documentation/Private/Documentation.cpp
+++ b/Engine/Source/Editor/Documentation/Private/Documentation.cpp
@@ -188,6 +188,11 @@ const TArray <FString>& FDocumentation::GetSourcePaths() const
 }
 
 TSharedRef< class SToolTip > FDocumentation::CreateToolTip(const TAttribute<FText>& Text, const TSharedPtr<SWidget>& OverrideContent, const FString& Link, const FString& ExcerptName) const
+{
+	return CreateToolTip(Text, OverrideContent, Link, ExcerptName, FText());
+}
+
+TSharedRef< class SToolTip > FDocumentation::CreateToolTip(const TAttribute<FText>& Text, const TSharedPtr<SWidget>& OverrideContent, const FString& Link, const FString& ExcerptName, const TAttribute<FText>& Shortcut) const
 {
 	TSharedPtr< SDocumentationToolTip > DocToolTip;
 
@@ -201,16 +206,18 @@ TSharedRef< class SToolTip > FDocumentation::CreateToolTip(const TAttribute<FTex
 		SAssignNew( DocToolTip, SDocumentationToolTip )
 		.DocumentationLink( Link )
 		.ExcerptName( ExcerptName )
+		.Shortcut(Shortcut)
 		[
 			OverrideContent.ToSharedRef()
 		];
 	}
 	else
 	{
-		SAssignNew( DocToolTip, SDocumentationToolTip )
-		.Text( Text )
-		.DocumentationLink( Link )
-		.ExcerptName( ExcerptName );
+		SAssignNew(DocToolTip, SDocumentationToolTip)
+		.Text(Text)
+		.DocumentationLink(Link)
+		.ExcerptName(ExcerptName)
+		.Shortcut(Shortcut);
 	}
 	
 	return SNew( SToolTip )
@@ -219,7 +226,8 @@ TSharedRef< class SToolTip > FDocumentation::CreateToolTip(const TAttribute<FTex
 		// Emulate text-only tool-tip styling that SToolTip uses when no custom content is supplied.  We want documentation tool-tips to 
 		// be styled just like text-only tool-tips
 		.BorderImage( FCoreStyle::Get().GetBrush("ToolTip.BrightBackground") )
-		.TextMargin(FMargin(11.0f))
+		// The documentation tooltip has padding built in, so that it can style its own background within the tooltip borders.
+		.TextMargin(FMargin(0.0f))
 		[
 			DocToolTip.ToSharedRef()
 		];
@@ -249,7 +257,8 @@ TSharedRef< class SToolTip > FDocumentation::CreateToolTip(const TAttribute<FTex
 		// Emulate text-only tool-tip styling that SToolTip uses when no custom content is supplied.  We want documentation tool-tips to 
 		// be styled just like text-only tool-tips
 		.BorderImage( FCoreStyle::Get().GetBrush("ToolTip.BrightBackground") )
-		.TextMargin(FMargin(11.0f))
+		// The documentation tooltip has padding built in, so that it can style its own background within the tooltip borders.
+		.TextMargin(FMargin(0.0f))
 		[
 			DocToolTip
 		];
diff --git a/Engine/Source/Editor/Documentation/Private/Documentation.h b/Engine/Source/Editor/Documentation/Private/Documentation.h
--- a/Engine/Source/Editor/Documentation/Private/Documentation.h
+++ b/Engine/Source/Editor/Documentation/Private/Documentation.h
@@ -48,6 +48,8 @@ public:
 
 	virtual TSharedRef< class SToolTip > CreateToolTip( const TAttribute<FText>& Text, const TSharedPtr<SWidget>& OverrideContent, const FString& Link, const FString& ExcerptName ) const override;
 	
+	virtual TSharedRef< class SToolTip > CreateToolTip(const TAttribute<FText>& Text, const TSharedPtr<SWidget>& OverrideContent, const FString& Link, const FString& ExcerptName, const TAttribute<FText>& Shortcut) const override;
+
 	virtual TSharedRef< class SToolTip > CreateToolTip(const TAttribute<FText>& Text, const TSharedRef<SWidget>& OverrideContent, const TSharedPtr<SVerticalBox>& DocVerticalBox, const FString& Link, const FString& ExcerptName) const override;
 
 	virtual bool RegisterBaseUrl(const FString& Id, const FString& Url) override;
diff --git a/Engine/Source/Editor/Documentation/Private/DocumentationModule.cpp b/Engine/Source/Editor/Documentation/Private/DocumentationModule.cpp
--- a/Engine/Source/Editor/Documentation/Private/DocumentationModule.cpp
+++ b/Engine/Source/Editor/Documentation/Private/DocumentationModule.cpp
@@ -49,11 +49,33 @@ public:
 
 private:
 
-	TSharedRef< SToolTip > ConstructDefaultToolTip( const TAttribute<FText>& ToolTipText, const TSharedPtr<SWidget>& OverrideContent, const TSharedPtr<const FUICommandInfo>& Action )
+	TSharedRef< SToolTip > ConstructDefaultToolTip( const TAttribute<FText>& ToolTipText, const TSharedPtr<SWidget>& OverrideContent, const TSharedPtr<const FUICommandInfo>& Action, bool ShowActionShortcut = false )
 	{
+		struct Local
+		{
+			/** Appends the key binding to the end of the provided ToolTip */
+			static FText GetCommandShortcut(TSharedPtr< const FUICommandInfo> Command)
+			{
+				if (Command.IsValid() && (Command->GetFirstValidChord()->IsValidChord()))
+				{
+					return Command->GetInputText();
+				}
+				else
+				{
+					return FText();
+				}
+			}
+		};
+
 		if ( Action.IsValid() )
 		{
-			return Documentation->CreateToolTip( ToolTipText, OverrideContent, FString( TEXT("Shared/") ) + Action->GetBindingContext().ToString(), Action->GetCommandName().ToString() );
+			TAttribute<FText> Shortcut;
+			if (ShowActionShortcut)
+			{
+				Shortcut = TAttribute<FText>::Create(TAttribute<FText>::FGetter::CreateStatic(&Local::GetCommandShortcut, Action));
+			}
+
+			return Documentation->CreateToolTip( ToolTipText, OverrideContent, FString( TEXT("Shared/") ) + Action->GetBindingContext().ToString(), Action->GetCommandName().ToString(), Shortcut);
 		} 
 		else
 		{
diff --git a/Engine/Source/Editor/Documentation/Private/DocumentationStyleSet.cpp b/Engine/Source/Editor/Documentation/Private/DocumentationStyleSet.cpp
new file mode 100644
--- /dev/null
+++ b/Engine/Source/Editor/Documentation/Private/DocumentationStyleSet.cpp
@@ -0,0 +1,61 @@
+// Copyright Epic Games, Inc. All Rights Reserved.
+
+#include "DocumentationStyleSet.h"
+#include "Styling/SlateStyleRegistry.h"
+#include "Styling/SlateStyleMacros.h"
+#include "Styling/SlateTypes.h"
+#include "Styling/AppStyle.h"
+#include "Styling/CoreStyle.h"
+#include "Styling/StyleColors.h"
+
+FName FDocumentationStyleSet::StyleName("FDocumentationStyleSet");
+TUniquePtr<FDocumentationStyleSet> FDocumentationStyleSet::Instance(nullptr);
+
+const FName& FDocumentationStyleSet::GetStyleSetName() const
+{
+	return StyleName;
+}
+
+const FDocumentationStyleSet& FDocumentationStyleSet::Get()
+{
+	if (!Instance.IsValid())
+	{
+		Instance = TUniquePtr<FDocumentationStyleSet>(new FDocumentationStyleSet);
+	}
+	return *(Instance.Get());
+}
+
+void FDocumentationStyleSet::Shutdown()
+{
+	Instance.Reset();
+}
+
+FDocumentationStyleSet::FDocumentationStyleSet()
+	: FSlateStyleSet(StyleName)
+{
+	const FVector2D Icon8x8(8.0f, 8.0f);
+	const FVector2D Icon16x16(16.0f, 16.0f);
+	const FVector2D Icon32x32(32.0f, 32.0f);
+	const FVector2D Icon128x128(128.0f, 128.0f);
+
+	Set("ToolTip.Header", new FSlateColorBrush(FStyleColors::Foreground));
+
+	Set("ToolTip.ScrollBox", FScrollBoxStyle());
+
+	Set("ToolTip.TopSeparator", new BORDER_BRUSH("Common/Selector", FMargin(0.0f, 1.0f, 0.0f, 0.0f), FStyleColors::Hover));
+
+	Set("ToolTip.KeybindBorder", new FSlateRoundedBoxBrush(FStyleColors::Hover2, 4.0f));
+	const FTextBlockStyle KeybindText = FTextBlockStyle(FAppStyle::GetWidgetStyle<FTextBlockStyle>("NormalText"))
+		.SetFont(DEFAULT_FONT("Bold", 10))
+		.SetColorAndOpacity(FLinearColor(1.0f, 1.0f, 1.0f));
+	Set("ToolTip.KeybindText", FTextBlockStyle(KeybindText));
+
+	Set("ToolTip.ToggleKeybindBorder", new FSlateRoundedBoxBrush(FStyleColors::Transparent, 4.0f, FLinearColor(0.1f, 0.1f, 0.1f), 1.0f));
+
+	FSlateStyleRegistry::RegisterSlateStyle(*this);
+}
+
+FDocumentationStyleSet::~FDocumentationStyleSet()
+{
+	FSlateStyleRegistry::UnRegisterSlateStyle(*this);
+}
diff --git a/Engine/Source/Editor/Documentation/Private/DocumentationStyleSet.h b/Engine/Source/Editor/Documentation/Private/DocumentationStyleSet.h
new file mode 100644
--- /dev/null
+++ b/Engine/Source/Editor/Documentation/Private/DocumentationStyleSet.h
@@ -0,0 +1,24 @@
+#pragma once
+
+#include "CoreMinimal.h"
+#include "Styling/SlateStyle.h"
+
+class FDocumentationStyleSet :
+    public FSlateStyleSet
+{
+public:
+	virtual const FName& GetStyleSetName() const override;
+
+	static const FDocumentationStyleSet& Get();
+	static void Shutdown();
+
+	~FDocumentationStyleSet();
+
+private:
+	FDocumentationStyleSet();
+
+	static FName StyleName;
+	static TUniquePtr<FDocumentationStyleSet> Instance;
+};
+
+
diff --git a/Engine/Source/Editor/Documentation/Private/SDocumentationToolTip.cpp b/Engine/Source/Editor/Documentation/Private/SDocumentationToolTip.cpp
--- a/Engine/Source/Editor/Documentation/Private/SDocumentationToolTip.cpp
+++ b/Engine/Source/Editor/Documentation/Private/SDocumentationToolTip.cpp
@@ -8,6 +8,7 @@
 #include "Framework/Application/SlateApplication.h"
 #include "Widgets/Text/STextBlock.h"
 #include "Widgets/Layout/SBox.h"
+#include "Widgets/Layout/SScrollBox.h"
 #include "Widgets/SToolTip.h"
 #include "Styling/AppStyle.h"
 #include "Editor/EditorPerProjectUserSettings.h"
@@ -23,6 +24,8 @@
 #include "Widgets/Input/SHyperlink.h"
 #include "Framework/Notifications/NotificationManager.h"
 #include "Widgets/Notifications/SNotificationList.h"
+#include "DocumentationStyleSet.h"
+#include "SPrimaryButton.h"
 
 void SDocumentationToolTip::Construct( const FArguments& InArgs )
 {
@@ -31,11 +34,13 @@ void SDocumentationToolTip::Construct( const FArguments& InArgs )
 	SubduedStyleInfo = FAppStyle::GetWidgetStyle<FTextBlockStyle>(InArgs._SubduedStyle);
 	HyperlinkTextStyleInfo = FAppStyle::GetWidgetStyle<FTextBlockStyle>(InArgs._HyperlinkTextStyle);
 	HyperlinkButtonStyleInfo = FAppStyle::GetWidgetStyle<FButtonStyle>(InArgs._HyperlinkButtonStyle);
+	KeybindStyleInfo = FDocumentationStyleSet::Get().GetWidgetStyle<FTextBlockStyle>("ToolTip.KeybindText");
 	ColorAndOpacity = InArgs._ColorAndOpacity;
 	DocumentationLink = InArgs._DocumentationLink;
-	IsDisplayingDocumentationLink = false;
 	bAddDocumentation = InArgs._AddDocumentation;
 	DocumentationMargin = InArgs._DocumentationMargin;
+	IsDisplayingDocumentationLink = false;
+	Shortcut = InArgs._Shortcut;
 
 	ExcerptName = InArgs._ExcerptName;
 	IsShowingFullTip = false;
@@ -47,15 +52,65 @@ void SDocumentationToolTip::Construct( const FArguments& InArgs )
 		OverrideContent = InArgs._Content.Widget;
 	}
 
+	SAssignNew(DocumentationControlBox, SHorizontalBox);
+	SAssignNew(FullTipContent, SBox);
 	ConstructSimpleTipContent();
 
 	ChildSlot
 	[
 		SAssignNew(WidgetContent, SBox)
+		.Padding(2.0f)
 		[
-			SimpleTipContent.ToSharedRef()
+			SNew(SVerticalBox)
+			+SVerticalBox::Slot()
+			.AutoHeight()
+			[
+				SimpleTipContent.ToSharedRef()
+			]
+			+ SVerticalBox::Slot()
+			.AutoHeight()
+			[
+				FullTipContent.ToSharedRef()
+			]
+			+ SVerticalBox::Slot()
+			.AutoHeight()
+			[
+				SNew(SBox)
+				.Visibility(this, &SDocumentationToolTip::GetControlVisibility)
+				[
+					SNew(SVerticalBox)
+					+ SVerticalBox::Slot()
+					.AutoHeight()
+					[
+						SNew(SBorder)
+						.Padding(0.0)
+						.BorderImage(FDocumentationStyleSet::Get().GetBrush("ToolTip.TopSeparator"))
+						[
+							SNew(SBox)
+							.HeightOverride(1.0f)
+						]
+					]
+					+ SVerticalBox::Slot()
+					.AutoHeight()
+					[
+						SNew(SBorder)
+						.Padding(9.0)
+						.BorderImage(FDocumentationStyleSet::Get().GetBrush("ToolTip.Header"))
+						[
+							DocumentationControlBox.ToSharedRef()
+						]
+					]
+				]
+			]
 		]
 	];
+
+	bIsInTransition = false;
+	TransitionStartTime = 0;
+	TransitionLength = 0.2f;
+	LastDesiredSize = SimpleTipContent->GetDesiredSize();
+	TransitionStartSize = SimpleTipContent->GetDesiredSize();
+	bFullTipContentIsReady = false;
 }
 
 void SDocumentationToolTip::ConstructSimpleTipContent()
@@ -85,31 +140,69 @@ void SDocumentationToolTip::ConstructSimpleTipContent()
 	}
 
 	TSharedPtr< SVerticalBox > VerticalBox;
+	TSharedPtr< SHorizontalBox > TextBox;
 	if ( !OverrideContent.IsValid() )
 	{
 		SAssignNew( SimpleTipContent, SBox )
 		[
-			SAssignNew( VerticalBox, SVerticalBox )
-			+SVerticalBox::Slot()
-			.FillHeight( 1.0f )
+			SNew(SBorder)
+			.BorderImage(this, &SDocumentationToolTip::GetSimpleTipBorderStyle)
+			.Padding(9.f)
+			[
+				SAssignNew( VerticalBox, SVerticalBox )
+				+SVerticalBox::Slot()
+				.FillHeight( 1.0f )
+				[
+					SAssignNew(TextBox, SHorizontalBox)
+
+					+SHorizontalBox::Slot()
+					.FillWidth(1.0f)
+					.VAlign(VAlign_Center)
+					[
+					SNew( STextBlock )
+					.Text( TextContent )
+					.TextStyle( &StyleInfo )
+					.ColorAndOpacity( ColorAndOpacity )
+					.WrapTextAt_Static( &SToolTip::GetToolTipWrapWidth )
+					]
+				]
+			]
+		];
+
+		TextBox->AddSlot()
+		.VAlign(VAlign_Center)
+		.AutoWidth()
+		[
+			SNew(SBox)
+			.Visibility(this, &SDocumentationToolTip::GetShortcutVisibility)
+			.Padding(9.0f, 0.0f, 0.0f, 0.0f)
 			[
-				SNew( STextBlock )
-				.Text( TextContent )
-				.TextStyle( &StyleInfo )
-				.ColorAndOpacity( ColorAndOpacity )
-				.WrapTextAt_Static( &SToolTip::GetToolTipWrapWidth )
+				SNew(SBorder)
+				.BorderImage(FDocumentationStyleSet::Get().GetBrush("ToolTip.KeybindBorder"))
+				.Padding(4.0f, 2.0f)
+				[
+					SNew(STextBlock)
+					.TextStyle(&KeybindStyleInfo)
+					.Text(Shortcut)
+				]
 			]
 		];
+	
 	}
 	else
 	{
 		SAssignNew( SimpleTipContent, SBox )
 		[
-			SAssignNew( VerticalBox, SVerticalBox )
-			+SVerticalBox::Slot()
-			.FillHeight( 1.0f )
+			SNew(SBorder)
+			.BorderImage(this, &SDocumentationToolTip::GetSimpleTipBorderStyle)
+			.Padding(9.f)
 			[
-				OverrideContent.ToSharedRef()
+				SAssignNew( VerticalBox, SVerticalBox )
+				+SVerticalBox::Slot()
+				.FillHeight( 1.0f )
+				[
+					OverrideContent.ToSharedRef()
+				]
 			]
 		];
 	}
@@ -124,36 +217,15 @@ void SDocumentationToolTip::AddDocumentation(TSharedPtr< SVerticalBox > Vertical
 {
 	if ( !DocumentationLink.IsEmpty() )
 	{
-		IsDisplayingDocumentationLink = GetDefault<UEditorPerProjectUserSettings>()->bDisplayDocumentationLink;
-
-		if ( IsDisplayingDocumentationLink )
-		{
-			FString OptionalExcerptName;
-			if ( !ExcerptName.IsEmpty() )
-			{ 
-				OptionalExcerptName = FString( TEXT(" [") ) + ExcerptName + TEXT("]");
-			}
-
-			VerticalBox->AddSlot()
-			.AutoHeight()
-			.Padding(0, 5, 0, 0)
-			.HAlign( HAlign_Center )
-			[
-				SNew( STextBlock )
-				.Text( FText::FromString(DocumentationLink + OptionalExcerptName) )
-				.TextStyle( &SubduedStyleInfo )
-			];
-		}
-
 		if ( !DocumentationPage.IsValid() )
 		{
 			DocumentationPage = IDocumentation::Get()->GetPage( DocumentationLink, NULL );
 		}
 
-		if ( DocumentationPage->HasExcerpt( ExcerptName ) )
+		if ( DocumentationPage->HasExcerpt( ExcerptName ))
 		{
-			FText MacShortcut = NSLOCTEXT("SToolTip", "MacRichTooltipShortcut", "(Command + Option)");
-			FText WinShortcut = NSLOCTEXT("SToolTip", "WinRichTooltipShortcut", "(Ctrl + Alt)");
+			FText MacShortcut = NSLOCTEXT("SToolTip", "MacRichTooltipShortcut", "Command + Option");
+			FText WinShortcut = NSLOCTEXT("SToolTip", "WinRichTooltipShortcut", "Ctrl + Alt");
 
 			FText KeyboardShortcut;
 #if PLATFORM_MAC
@@ -164,34 +236,77 @@ void SDocumentationToolTip::AddDocumentation(TSharedPtr< SVerticalBox > Vertical
 
 			VerticalBox->AddSlot()
 			.AutoHeight()
-			.HAlign( HAlign_Center )
-			.Padding(0, 5, 0, 0)
+			.HAlign( HAlign_Right )
 			[
-				SNew( STextBlock )
-				.TextStyle( &SubduedStyleInfo )
-				.Text( FText::Format( NSLOCTEXT( "SToolTip", "AdvancedToolTipMessage", "hold {0} for more" ), KeyboardShortcut) )
+				SNew(SBox)
+				.Visibility(this, &SDocumentationToolTip::GetPromptVisibility)
+				.Padding(0, 9, 0, 0)
+				[
+					SNew(SHorizontalBox)
+					+SHorizontalBox::Slot()
+					.VAlign(VAlign_Center)
+					.AutoWidth()
+					.Padding(0.0f, 0.0f, 5.0f, 0.0f)
+					[
+						SNew( STextBlock )
+						.TextStyle( &SubduedStyleInfo )
+						.Text( NSLOCTEXT( "SToolTip", "AdvancedToolTipMessage", "Learn more: hold" ) )
+					]
+					+SHorizontalBox::Slot()
+					.VAlign(VAlign_Center)
+					.AutoWidth()
+					[
+						SNew( SBorder )
+						.BorderImage(FDocumentationStyleSet::Get().GetBrush("ToolTip.ToggleKeybindBorder"))
+						.Padding(4.0f, 2.0f)
+						[
+							SNew(STextBlock)
+							.TextStyle(&SubduedStyleInfo)
+							.Text(FText::Format(NSLOCTEXT("SToolTip", "AdvancedToolTipKeybind", "{0}"), KeyboardShortcut))
+							.Visibility(this, &SDocumentationToolTip::GetPromptVisibility)
+						]
+					]
+				]
 			];
 		}
-		else
+
+		SAssignNew(DocumentationControlBox, SHorizontalBox);
+		IsDisplayingDocumentationLink = GetDefault<UEditorPerProjectUserSettings>()->bDisplayDocumentationLink;
+		if (IsDisplayingDocumentationLink)
 		{
-			if ( IsDisplayingDocumentationLink && FSlateApplication::Get().SupportsSourceAccess() )
+
+			FString OptionalExcerptName;
+			if (!ExcerptName.IsEmpty())
+			{
+				OptionalExcerptName = FString(TEXT(" [")) + ExcerptName + TEXT("]");
+			}
+			DocumentationControlBox->AddSlot()
+			.FillWidth(1.0f)
+			.HAlign(HAlign_Left)
+			.Padding(0, 0, 9, 0)
+			[
+				SNew(STextBlock)
+					.Text(FText::FromString(DocumentationLink + OptionalExcerptName))
+					.TextStyle(&SubduedStyleInfo)
+			];
+
+			if (!DocumentationPage->HasExcerpt(ExcerptName) && FSlateApplication::Get().SupportsSourceAccess())
 			{
-				FString DocPath = FDocumentationLink::ToSourcePath( DocumentationLink, FInternationalization::Get().GetCurrentCulture() );
-				if ( !FPaths::FileExists(DocPath) )
+				FString DocPath = FDocumentationLink::ToSourcePath(DocumentationLink, FInternationalization::Get().GetCurrentCulture());
+				if (!FPaths::FileExists(DocPath))
 				{
 					DocPath = FPaths::ConvertRelativePathToFull(DocPath);
 				}
 
-				VerticalBox->AddSlot()
-				.AutoHeight()
-				.Padding(0, 5, 0, 0)
-				.HAlign( HAlign_Center )
+				DocumentationControlBox->AddSlot()
+				.AutoWidth()
+				.HAlign(HAlign_Right)
 				[
-					SNew( SHyperlink )
-					.Text( NSLOCTEXT( "SToolTip", "EditDocumentationMessage_Create", "create" ) )
-					.TextStyle( &HyperlinkTextStyleInfo )
-					.UnderlineStyle( &HyperlinkButtonStyleInfo )
-					.OnNavigate( this, &SDocumentationToolTip::CreateExcerpt, DocPath, ExcerptName )
+					SNew(SHyperlink)
+						.Text(NSLOCTEXT("SToolTip", "EditDocumentationMessage_Create", "create"))
+						.TextStyle(&HyperlinkTextStyleInfo)
+						.UnderlineStyle(&HyperlinkButtonStyleInfo)
+						.OnNavigate(this, &SDocumentationToolTip::CreateExcerpt, DocPath, ExcerptName)
 				];
 			}
 		}
@@ -299,18 +414,29 @@ void SDocumentationToolTip::ConstructFullTipContent()
 		if ( Excerpts[ ExcerptIndex ].Content.IsValid() )
 		{
 			TSharedPtr< SVerticalBox > Box;
-			FullTipContent = 
-				SNew(SBox)
-				.Padding(DocumentationMargin)
+			TSharedPtr< SWidget > FullTipBox = SNew(SBox)
+			.Visibility(this, &SDocumentationToolTip::GetFullTipVisibility)
+			.Padding(DocumentationMargin)
+			[
+				SAssignNew(Box, SVerticalBox)
+				+ SVerticalBox::Slot()
+				.HAlign(HAlign_Center)
+				.AutoHeight()
+				.MaxHeight(900.f)
 				[
-					SAssignNew(Box, SVerticalBox)
-					+ SVerticalBox::Slot()
-					.HAlign(HAlign_Center)
-					.AutoHeight()
+					SNew(SBorder)
+					.BorderImage(FDocumentationStyleSet::Get().GetBrush("ToolTip.TopSeparator"))
 					[
-						Excerpts[ExcerptIndex].Content.ToSharedRef()
+						SNew(SScrollBox)
+						.Style(FDocumentationStyleSet::Get(), "ToolTip.ScrollBox")
+						+SScrollBox::Slot()
+						.Padding(9.0f)
+						[
+							Excerpts[ExcerptIndex].Content.ToSharedRef()
+						]
 					]
-				];
+				]
+			];
 
 			FString* FullDocumentationLink = Excerpts[ ExcerptIndex ].Variables.Find( TEXT("ToolTipFullLink") );
 			FString* ExcerptBaseUrl = Excerpts[ExcerptIndex].Variables.Find(TEXT("BaseUrl"));
@@ -323,40 +449,65 @@ void SDocumentationToolTip::ConstructFullTipContent()
 				}
 
 				Box->AddSlot()
-				.HAlign( HAlign_Center )
 				.AutoHeight()
 				[
-					SNew( SHyperlink )
-						.Text( NSLOCTEXT( "SToolTip", "GoToFullDocsLinkMessage", "see full documentation" ) )
-						.TextStyle( &HyperlinkTextStyleInfo )
-						.UnderlineStyle( &HyperlinkButtonStyleInfo )
-						.OnNavigate_Static([](FString Link, FString BaseUrl) {
+					SNew(SBorder)
+					.BorderImage(FDocumentationStyleSet::Get().GetBrush("ToolTip.TopSeparator"))
+					.Padding(0.0f)
+					[
+						SNew(SBox)
+						.HeightOverride(1.0f)
+					]
+				];
+
+				Box->AddSlot()
+				.AutoHeight()
+				[
+					SNew(SBorder)
+					.Padding(9.0f)
+					.BorderImage(FDocumentationStyleSet::Get().GetBrush("ToolTip.Header"))
+					[
+						SNew(SHorizontalBox)
+						+SHorizontalBox::Slot()
+						.FillWidth(1.0f)
+						.HAlign(HAlign_Right)
+						[
+							SNew(SPrimaryButton)
+							.Icon(FAppStyle::Get().GetBrush("Icons.Help"))
+							.Text(NSLOCTEXT("SToolTip", "LearnMoreButton", "Learn More Online"))
+							.OnClicked_Static([](FString Link, FString BaseUrl) -> FReply {
 								if (!IDocumentation::Get()->Open(Link, FDocumentationSourceInfo(TEXT("rich_tooltips")), BaseUrl))
 								{
 									FNotificationInfo Info(NSLOCTEXT("SToolTip", "FailedToOpenLink", "Failed to Open Link"));
 									FSlateNotificationManager::Get().AddNotification(Info);
 								}
-							}, *FullDocumentationLink, BaseUrl)
+								return FReply::Handled();
+								}, *FullDocumentationLink, BaseUrl)
+						]
+					]
 				];
 			}
 
-			if ( GetDefault<UEditorPerProjectUserSettings>()->bDisplayDocumentationLink && FSlateApplication::Get().SupportsSourceAccess() )
+			if (IsDisplayingDocumentationLink && FSlateApplication::Get().SupportsSourceAccess() )
 			{
-				Box->AddSlot()
-				.AutoHeight()
-				.HAlign( HAlign_Center )
+				DocumentationControlBox->AddSlot()
+				.AutoWidth()
+				.HAlign( HAlign_Right )
 				[
 					SNew( SHyperlink )
-						.Text( NSLOCTEXT( "SToolTip", "EditDocumentationMessage_Edit", "edit" ) )
-						.TextStyle( &HyperlinkTextStyleInfo )
-						.UnderlineStyle( &HyperlinkButtonStyleInfo )
-						// todo: needs to update to point to the "real" source file used for the excerpt
-						.OnNavigate_Static([](FString Link, int32 LineNumber) {
-								ISourceCodeAccessModule& SourceCodeAccessModule = FModuleManager::LoadModuleChecked<ISourceCodeAccessModule>("SourceCodeAccess");
-								SourceCodeAccessModule.GetAccessor().OpenFileAtLine(Link, LineNumber);
-							}, FPaths::ConvertRelativePathToFull(FDocumentationLink::ToSourcePath(DocumentationLink, FInternationalization::Get().GetCurrentCulture())), Excerpts[ExcerptIndex].LineNumber)
+					.Text( NSLOCTEXT( "SToolTip", "EditDocumentationMessage_Edit", "edit" ) )
+					.TextStyle( &HyperlinkTextStyleInfo )
+					.UnderlineStyle( &HyperlinkButtonStyleInfo )
+					// todo: needs to update to point to the "real" source file used for the excerpt
+					.OnNavigate_Static([](FString Link, int32 LineNumber) {
+							ISourceCodeAccessModule& SourceCodeAccessModule = FModuleManager::LoadModuleChecked<ISourceCodeAccessModule>("SourceCodeAccess");
+							SourceCodeAccessModule.GetAccessor().OpenFileAtLine(Link, LineNumber);
+						}, FPaths::ConvertRelativePathToFull(FDocumentationLink::ToSourcePath(DocumentationLink, FInternationalization::Get().GetCurrentCulture())), Excerpts[ExcerptIndex].LineNumber)
 				];
 			}
+			
+			FullTipContent->SetContent(FullTipBox.ToSharedRef());
+			bFullTipContentIsReady = true;
 		}
 	}
 }
@@ -364,7 +515,8 @@ void SDocumentationToolTip::ConstructFullTipContent()
 FReply SDocumentationToolTip::ReloadDocumentation()
 {
 	SimpleTipContent.Reset();
-	FullTipContent.Reset();
+	DocumentationControlBox.Reset();
+	bFullTipContentIsReady = false;
 
 	ConstructSimpleTipContent();
 
@@ -383,12 +535,30 @@ FReply SDocumentationToolTip::ReloadDocumentation()
 
 void SDocumentationToolTip::Tick( const FGeometry& AllottedGeometry, const double InCurrentTime, const float InDeltaTime )
 {
+	LastDesiredSize = WidgetContent->GetDesiredSize();
+
 	const FModifierKeysState ModifierKeys = FSlateApplication::Get().GetModifierKeys();
 	const bool NeedsUpdate = IsDisplayingDocumentationLink != GetDefault<UEditorPerProjectUserSettings>()->bDisplayDocumentationLink;
+	if (TransitionStartTime > 0)
+	{
+		bIsInTransition = InCurrentTime - TransitionStartTime <= TransitionLength;
+		if (bIsInTransition) {
+			TransitionPercentage = (InCurrentTime - TransitionStartTime) / TransitionLength;
+			FVector2D TransitionEndSize = WidgetContent->GetDesiredSize();
+			if (TransitionEndSize.Y > TransitionStartSize.Y)
+			{
+				LastDesiredSize = ((TransitionEndSize - TransitionStartSize) * FMath::InterpEaseOut<float>(0.f, 1.f, TransitionPercentage, 4.f)) + TransitionStartSize;
+			}
+			else
+			{
+				LastDesiredSize = TransitionStartSize - ((TransitionStartSize - TransitionEndSize) * FMath::InterpEaseOut<float>(0.f, 1.f, TransitionPercentage, 4.f));
+			}
+		}
+	}
 
 	if ( !IsShowingFullTip && ModifierKeys.IsAltDown() && ModifierKeys.IsControlDown() )
 	{
-		if ( !FullTipContent.IsValid() && DocumentationPage.IsValid() && DocumentationPage->HasExcerpt( ExcerptName ) )
+		if ( !bFullTipContentIsReady && DocumentationPage.IsValid() && DocumentationPage->HasExcerpt(ExcerptName))
 		{
 			ConstructFullTipContent();
 		}
@@ -397,11 +567,8 @@ void SDocumentationToolTip::Tick( const FGeometry& AllottedGeometry, const doubl
 			ReloadDocumentation();
 		}
 
-		if ( FullTipContent.IsValid() )
+		if ( bFullTipContentIsReady)
 		{
-			WidgetContent->SetContent( FullTipContent.ToSharedRef() );
-			IsShowingFullTip = true;
-
 			// Analytics event
 			if (FEngineAnalytics::IsAvailable())
 			{
@@ -411,6 +578,10 @@ void SDocumentationToolTip::Tick( const FGeometry& AllottedGeometry, const doubl
 
 				FEngineAnalytics::GetProvider().RecordEvent(TEXT("Editor.Usage.Documentation.FullTooltipShown"), Params);
 			}
+			bIsInTransition = true;
+			IsShowingFullTip = true;
+			TransitionStartTime = InCurrentTime;
+			TransitionStartSize = LastDesiredSize;
 		}
 	}
 	else if ( ( IsShowingFullTip || NeedsUpdate )  && ( !ModifierKeys.IsAltDown() || !ModifierKeys.IsControlDown() ) )
@@ -421,8 +592,10 @@ void SDocumentationToolTip::Tick( const FGeometry& AllottedGeometry, const doubl
 			IsDisplayingDocumentationLink = GetDefault<UEditorPerProjectUserSettings>()->bDisplayDocumentationLink;
 		}
 
-		WidgetContent->SetContent( SimpleTipContent.ToSharedRef() );
+		bIsInTransition = true;
 		IsShowingFullTip = false;
+		TransitionStartTime = InCurrentTime;
+		TransitionStartSize = LastDesiredSize;
 	}
 }
 
@@ -431,3 +604,53 @@ bool SDocumentationToolTip::IsInteractive() const
 	const FModifierKeysState ModifierKeys = FSlateApplication::Get().GetModifierKeys();
 	return ( DocumentationPage.IsValid() && ModifierKeys.IsAltDown() && ModifierKeys.IsControlDown() );
 }
+
+FVector2D SDocumentationToolTip::ComputeDesiredSize(float LayoutScaleMultiplier) const
+{
+	return LastDesiredSize;
+}
+
+EVisibility SDocumentationToolTip::GetFullTipVisibility() const
+{
+	if (IsShowingFullTip)
+	{
+		return EVisibility::Visible;
+	}
+	return EVisibility::Collapsed;
+}
+
+EVisibility SDocumentationToolTip::GetPromptVisibility() const
+{
+	if (IsShowingFullTip)
+	{
+		return EVisibility::Collapsed;
+	}
+	return EVisibility::Visible;
+}
+
+EVisibility SDocumentationToolTip::GetControlVisibility() const
+{
+	if (IsDisplayingDocumentationLink && (IsShowingFullTip || !DocumentationPage.IsValid() || !DocumentationPage->HasExcerpt(ExcerptName)))
+	{
+		return EVisibility::Visible;
+	}
+	return EVisibility::Collapsed;
+}
+
+EVisibility SDocumentationToolTip::GetShortcutVisibility() const
+{
+	if ((Shortcut.IsSet() || Shortcut.IsBound()) && !Shortcut.Get().IsEmpty())
+	{
+		return EVisibility::Visible;
+	}
+	return EVisibility::Collapsed;
+}
+
+const FSlateBrush* SDocumentationToolTip::GetSimpleTipBorderStyle() const
+{
+	if (IsShowingFullTip)
+	{
+		return FDocumentationStyleSet::Get().GetBrush("ToolTip.Header");
+	}
+	return FAppStyle::GetBrush("");
+}
diff --git a/Engine/Source/Editor/Documentation/Private/UDNParser.cpp b/Engine/Source/Editor/Documentation/Private/UDNParser.cpp
--- a/Engine/Source/Editor/Documentation/Private/UDNParser.cpp
+++ b/Engine/Source/Editor/Documentation/Private/UDNParser.cpp
@@ -72,6 +72,8 @@ static const FString AssetLinkSpecifier(TEXT("ASSETLINK:"));
 
 }
 
+static const int IndentWidth = 32;
+
 TSharedRef< FUDNParser > FUDNParser::Create( const TSharedPtr< FParserConfiguration >& ParserConfig, const FDocumentationStyle& Style ) 
 {
 	TSharedPtr< FParserConfiguration > FinalParserConfig = ParserConfig;
@@ -789,6 +791,38 @@ void FUDNParser::AddContentToExcerpt(TSharedPtr<SVerticalBox> Box, const FString
 	}
 }
 
+void FUDNParser::AddCaptionToExcerpt(TSharedPtr<SVerticalBox> Box, const FString& ContentSource, FExcerpt& Excerpt)
+{
+	if (!ContentSource.IsEmpty())
+	{
+		AppendExcerpt(Box,
+			SNew(SBox)
+			.WidthOverride(ContentWidth)
+			[
+				SNew(SHorizontalBox)
+				+ SHorizontalBox::Slot()
+				.AutoWidth()
+				[
+					SNew(SBox)
+					.WidthOverride(IndentWidth)
+				]
+				+SHorizontalBox::Slot()
+				.FillWidth(1.0f)
+				.HAlign(HAlign_Right)
+				[
+					SNew(STextBlock)
+					.Text(FText::Format(LOCTEXT("ImageCaptionLabel", "Image: {0}"), FText::FromString(ContentSource)))
+					.TextStyle(FAppStyle::Get(), Style.ItalicContentStyleName)
+					.WrapTextAt(WrapAt.Get() - IndentWidth)
+				]
+			]
+		);
+
+		AddLineSeperator(Excerpt);
+		Excerpt.RichText += FString::Printf(TEXT("<TextStyle Style=\"%s\">%s</>"), *Style.ContentStyleName.ToString(), *ContentSource);
+	}
+}
+
 void FUDNParser::AddListItemToExcerpt(TSharedPtr<SVerticalBox> Box, const FString& LeftContentSource, const FString& RightContentSource, FExcerpt& Excerpt)
 {
 	if (!LeftContentSource.IsEmpty() && !RightContentSource.IsEmpty())
@@ -807,7 +841,7 @@ void FUDNParser::AddListItemToExcerpt(TSharedPtr<SVerticalBox> Box, const FStrin
 					.AutoWidth()
 					[
 						SNew(SBox)
-						.MinDesiredWidth(40.f)
+						.WidthOverride(IndentWidth)
 						.Padding(FMargin(16.f, 0, 0, 0))
 						.HAlign(HAlign_Left)
 						[
@@ -817,12 +851,12 @@ void FUDNParser::AddListItemToExcerpt(TSharedPtr<SVerticalBox> Box, const FStrin
 						]
 					]
 					+ SHorizontalBox::Slot()
-					.FillWidth(100)
+					.FillWidth(100.0f)
 					[
 						SNew(STextBlock)
 						.Text(FText::FromString(RightContentSource))
 						.TextStyle(FAppStyle::Get(), Style.ContentStyleName)
-						.AutoWrapText(true)
+						.WrapTextAt(WrapAt.Get() - IndentWidth)
 					]
 				]
 			];
@@ -1069,6 +1103,11 @@ TSharedRef< SWidget > FUDNParser::GenerateExcerptContent( const FString& InLink,
 						);
 
 					DynamicBrushesUsed.AddUnique(DynamicBrush);
+
+					if (!Line.AdditionalContent[0].IsEmpty())
+					{
+						AddCaptionToExcerpt(Box, Line.AdditionalContent[0], Excerpt);
+					}
 				}
 
 				AddLineSeperator(Excerpt);
diff --git a/Engine/Source/Editor/Documentation/Private/UDNParser.h b/Engine/Source/Editor/Documentation/Private/UDNParser.h
--- a/Engine/Source/Editor/Documentation/Private/UDNParser.h
+++ b/Engine/Source/Editor/Documentation/Private/UDNParser.h
@@ -164,6 +164,9 @@ private:
 	/** Adds the content text source to the scrollbox */
 	void AddContentToExcerpt(TSharedPtr<SVerticalBox> Box, const FString& ContentSource, FExcerpt& Excerpt);
 	
+	/** Adds the content text source to the scrollbox as a caption for an image */
+	void AddCaptionToExcerpt(TSharedPtr<SVerticalBox> Box, const FString& ContentSource, FExcerpt& Excerpt);
+
 	/** Adds the content text sources to the scrollbox in two columns, emulating a list item with hanging indent */
 	void AddListItemToExcerpt(TSharedPtr<SVerticalBox> Box, const FString& LeftContentSource, const FString& RightContentSource, FExcerpt& Excerpt);
 
diff --git a/Engine/Source/Editor/Documentation/Public/IDocumentation.h b/Engine/Source/Editor/Documentation/Public/IDocumentation.h
--- a/Engine/Source/Editor/Documentation/Public/IDocumentation.h
+++ b/Engine/Source/Editor/Documentation/Public/IDocumentation.h
@@ -20,6 +20,7 @@ public:
 	FDocumentationStyle()
 		: ContentStyleName(TEXT("Documentation.Content"))
 		, BoldContentStyleName(TEXT("Documentation.BoldContent"))
+		, ItalicContentStyleName(TEXT("Documentation.ItalicContent"))
 		, NumberedContentStyleName(TEXT("Documentation.NumberedContent"))
 		, Header1StyleName(TEXT("Documentation.Header1"))
 		, Header2StyleName(TEXT("Documentation.Header2"))
@@ -44,6 +45,13 @@ public:
 		return *this;
 	}
 
+	/** Set the italic content style for this documentation */
+	FDocumentationStyle& ItalicContentStyle(const FName& InName)
+	{
+		ItalicContentStyleName = InName;
+		return *this;
+	}
+
 	/** Set the numbered content style for this documentation */
 	FDocumentationStyle& NumberedContentStyle(const FName& InName) 
 	{
@@ -99,6 +107,9 @@ public:
 	/** Bold content text style */
 	FName BoldContentStyleName;
 
+	/** Italic content text style */
+	FName ItalicContentStyleName;
+
 	/** Numbered content text style */
 	FName NumberedContentStyleName;
 
@@ -201,6 +212,8 @@ public:
 
 	virtual TSharedRef< class SToolTip > CreateToolTip( const TAttribute<FText>& Text, const TSharedPtr<SWidget>& OverrideContent, const FString& Link, const FString& ExcerptName ) const = 0;
 
+	virtual TSharedRef< class SToolTip > CreateToolTip(const TAttribute<FText>& Text, const TSharedPtr<SWidget>& OverrideContent, const FString& Link, const FString& ExcerptName, const TAttribute<FText>& Shortcut) const = 0;
+
 	virtual TSharedRef< class SToolTip > CreateToolTip(const TAttribute<FText>& Text, const TSharedRef<SWidget>& OverrideContent, const TSharedPtr<class SVerticalBox>& DocVerticalBox, const FString& Link, const FString& ExcerptName) const = 0;
 
 	virtual bool RegisterBaseUrl(const FString& Id, const FString& Url) = 0;
diff --git a/Engine/Source/Editor/Documentation/Public/SDocumentationToolTip.h b/Engine/Source/Editor/Documentation/Public/SDocumentationToolTip.h
--- a/Engine/Source/Editor/Documentation/Public/SDocumentationToolTip.h
+++ b/Engine/Source/Editor/Documentation/Public/SDocumentationToolTip.h
@@ -15,6 +15,7 @@
 #include "UObject/NameTypes.h"
 #include "Widgets/DeclarativeSyntaxSupport.h"
 #include "Widgets/SCompoundWidget.h"
+#include "Widgets/SBoxPanel.h"
 
 class IDocumentationPage;
 class SBox;
@@ -29,6 +30,7 @@ public:
 
 	SLATE_BEGIN_ARGS( SDocumentationToolTip )
 		: _Text()
+		, _Shortcut()
 		, _Style( TEXT("Documentation.SDocumentationTooltip") )
 		, _SubduedStyle( TEXT("Documentation.SDocumentationTooltipSubdued") )
 		, _HyperlinkTextStyle( TEXT("Documentation.SDocumentationTooltipHyperlinkText") )
@@ -42,6 +44,9 @@ public:
 		/** The text displayed in this tool tip */
 		SLATE_ATTRIBUTE( FText, Text )
 
+		/** The keyboard shortcut to display for the action */
+		SLATE_ATTRIBUTE(FText, Shortcut)
+
 		/** The text style to use for this tool tip */
 		SLATE_ARGUMENT( FName, Style )
 		
@@ -97,6 +102,8 @@ public:
 	 */
 	void AddDocumentation(TSharedPtr< SVerticalBox > VerticalBox);
 
+	virtual FVector2D ComputeDesiredSize(float LayoutScaleMultiplier) const override;
+
 private:
 
 	void ConstructSimpleTipContent();
@@ -107,14 +114,22 @@ private:
 
 	void CreateExcerpt( FString FileSource, FString ExcerptName );
 
+	EVisibility GetFullTipVisibility() const;
+	EVisibility GetPromptVisibility() const;
+	EVisibility GetControlVisibility() const;
+	EVisibility GetShortcutVisibility() const;
+	const FSlateBrush* GetSimpleTipBorderStyle() const;
+
 private:
 
 	/** Text block widget */
 	TAttribute< FText > TextContent;
+	TAttribute< FText > Shortcut;
 	TSharedPtr< SWidget > OverrideContent;
 	FTextBlockStyle StyleInfo;
 	FTextBlockStyle SubduedStyleInfo;
 	FTextBlockStyle HyperlinkTextStyleInfo;
+	FTextBlockStyle KeybindStyleInfo;
 	FButtonStyle HyperlinkButtonStyleInfo;
 	TAttribute< FSlateColor > ColorAndOpacity;
 
@@ -128,11 +143,21 @@ private:
 	TSharedPtr< SWidget > SimpleTipContent;
 	bool IsDisplayingDocumentationLink;
 
-	TSharedPtr< SWidget > FullTipContent;
+	TSharedPtr< SBox > FullTipContent;
+	TSharedPtr< SHorizontalBox > DocumentationControlBox;
 
 	TSharedPtr< IDocumentationPage > DocumentationPage;
 	bool IsShowingFullTip;
 
 	bool bAddDocumentation;
 	FMargin DocumentationMargin;
+
+	FVector2D LastDesiredSize;
+	bool bIsInTransition;
+	double TransitionStartTime;
+	float TransitionLength;
+	float TransitionPercentage;
+	FVector2D TransitionStartSize;
+
+	bool bFullTipContentIsReady;
 };
diff --git a/Engine/Source/Editor/EditorStyle/Private/StarshipStyle.cpp b/Engine/Source/Editor/EditorStyle/Private/StarshipStyle.cpp
--- a/Engine/Source/Editor/EditorStyle/Private/StarshipStyle.cpp
+++ b/Engine/Source/Editor/EditorStyle/Private/StarshipStyle.cpp
@@ -2843,7 +2843,7 @@ void FStarshipEditorStyle::FStyle::SetupTutorialStyles()
 	// Documentation defaults
 	const FTextBlockStyle DocumentationText = FTextBlockStyle(NormalText)
 		.SetColorAndOpacity( FLinearColor::Black )
-		.SetFont(DEFAULT_FONT( "Regular", 11 ));
+		.SetFont(DEFAULT_FONT( "Regular", 10 ));
 	Set("Documentation.Text", FTextBlockStyle(DocumentationText));
 
 	const FTextBlockStyle DocumentationHyperlinkText = FTextBlockStyle(DocumentationText)
@@ -2876,6 +2876,8 @@ void FStarshipEditorStyle::FStyle::SetupTutorialStyles()
 		Set("Documentation.NumberedContent",  FTextBlockStyle(DocumentationText));
 		Set( "Documentation.BoldContent", FTextBlockStyle(DocumentationText)
 			.SetTypefaceFontName(TEXT("Bold")));
+		Set("Documentation.ItalicContent", FTextBlockStyle(DocumentationText)
+			.SetTypefaceFontName(TEXT("Italic")));
 
 		Set("Documentation.Header1", FTextBlockStyle(DocumentationHeaderText)
 			.SetFontSize(32));
diff --git a/Engine/Source/Runtime/Slate/Private/Framework/MultiBox/MultiBox.cpp b/Engine/Source/Runtime/Slate/Private/Framework/MultiBox/MultiBox.cpp
--- a/Engine/Source/Runtime/Slate/Private/Framework/MultiBox/MultiBox.cpp
+++ b/Engine/Source/Runtime/Slate/Private/Framework/MultiBox/MultiBox.cpp
@@ -97,8 +97,27 @@ FMultiBoxSettings::FMultiBoxSettings()
 	ResetToolTipConstructor();
 }
 
-TSharedRef< SToolTip > FMultiBoxSettings::ConstructDefaultToolTip( const TAttribute<FText>& ToolTipText, const TSharedPtr<SWidget>& OverrideContent, const TSharedPtr<const FUICommandInfo>& Action )
+TSharedRef< SToolTip > FMultiBoxSettings::ConstructDefaultToolTip( const TAttribute<FText>& ToolTipText, const TSharedPtr<SWidget>& OverrideContent, const TSharedPtr<const FUICommandInfo>& Action, bool ShowActionShortcut )
 {
+	struct Local
+	{
+		/** Appends the key binding to the end of the provided ToolTip */
+		static FText AppendKeyBindingToToolTip(const TAttribute<FText> ToolTip, TSharedPtr< const FUICommandInfo> Command)
+		{
+			if (Command.IsValid() && (Command->GetFirstValidChord()->IsValidChord()))
+			{
+				FFormatNamedArguments Args;
+				Args.Add(TEXT("ToolTipDescription"), ToolTip.Get());
+				Args.Add(TEXT("Keybinding"), Command->GetInputText());
+				return FText::Format(NSLOCTEXT("ToolBar", "ToolTip + Keybinding", "{ToolTipDescription} ({Keybinding})"), Args);
+			}
+			else
+			{
+				return ToolTip.Get();
+			}
+		}
+	};
+
 	if ( OverrideContent.IsValid() )
 	{
 		return SNew( SToolTip )
@@ -107,7 +126,13 @@ TSharedRef< SToolTip > FMultiBoxSettings::ConstructDefaultToolTip( const TAttrib
 		];
 	}
 
-	return SNew( SToolTip ).Text( ToolTipText );
+	TAttribute<FText> ActualToolTip = ToolTipText;
+	if (Action.IsValid())
+	{
+		ActualToolTip = TAttribute< FText >::Create(TAttribute<FText>::FGetter::CreateStatic(&Local::AppendKeyBindingToToolTip, ToolTipText, Action));
+	}
+
+	return SNew( SToolTip ).Text(ActualToolTip);
 }
 
 void FMultiBoxSettings::ResetToolTipConstructor()
diff --git a/Engine/Source/Runtime/Slate/Private/Framework/MultiBox/SMenuEntryBlock.cpp b/Engine/Source/Runtime/Slate/Private/Framework/MultiBox/SMenuEntryBlock.cpp
--- a/Engine/Source/Runtime/Slate/Private/Framework/MultiBox/SMenuEntryBlock.cpp
+++ b/Engine/Source/Runtime/Slate/Private/Framework/MultiBox/SMenuEntryBlock.cpp
@@ -650,7 +650,7 @@ TSharedRef< SWidget > SMenuEntryBlock::BuildMenuEntryWidget( const FMenuEntryBui
 		.ClickMethod( ButtonClickMethod )
 		.ContentPadding(StyleSet->GetMargin(StyleName,".Block.Padding"))
 		// Pass along the block's tool-tip string
-		.ToolTip( FMultiBoxSettings::ToolTipConstructor.Execute(EntryToolTip, nullptr, UICommand ) )
+		.ToolTip( FMultiBoxSettings::ToolTipConstructor.Execute(EntryToolTip, nullptr, UICommand, /*ShowActionShortcut=*/ false ) )
 		// Bind the button's "on clicked" event to our object's method for this
 		.OnClicked(this, &SMenuEntryBlock::OnMenuItemButtonClicked)
 		[
@@ -900,7 +900,7 @@ TSharedRef< SWidget> SMenuEntryBlock::BuildSubMenuWidget( const FMenuEntryBuildP
 			// Create a button
 			SNew( SSubMenuButton )
 			// Pass along the block's tool-tip string
-			.ToolTip( FMultiBoxSettings::ToolTipConstructor.Execute(EntryToolTip, nullptr, UICommand ) )
+			.ToolTip( FMultiBoxSettings::ToolTipConstructor.Execute(EntryToolTip, nullptr, UICommand, /*ShowActionShortcut=*/ false) )
 			// Style to use
 			.ButtonStyle( &StyleSet->GetWidgetStyle<FButtonStyle>( ISlateStyle::Join( StyleName, ".Button" ) ) )
 			.ContentPadding(StyleSet->GetMargin(StyleName, ".Block.Padding"))
diff --git a/Engine/Source/Runtime/Slate/Private/Framework/MultiBox/SToolBarButtonBlock.cpp b/Engine/Source/Runtime/Slate/Private/Framework/MultiBox/SToolBarButtonBlock.cpp
--- a/Engine/Source/Runtime/Slate/Private/Framework/MultiBox/SToolBarButtonBlock.cpp
+++ b/Engine/Source/Runtime/Slate/Private/Framework/MultiBox/SToolBarButtonBlock.cpp
@@ -212,10 +212,6 @@ void SToolBarButtonBlock::BuildMultiBlockWidget(const ISlateStyle* StyleSet, con
 	{
 		ActualToolTip = UICommand.IsValid() ? UICommand->GetDescription() : FText::GetEmpty();
 	}
-
-	// If a key is bound to the command, append it to the tooltip text.
-	TWeakPtr<const FUICommandInfo> Action = ToolBarButtonBlock->GetAction();
-	ActualToolTip = TAttribute< FText >::Create(TAttribute<FText>::FGetter::CreateStatic(&Local::AppendKeyBindingToToolTip, ActualToolTip, Action ) );
 	
 	// If we were supplied an image than go ahead and use that, otherwise we use a null widget
 	TSharedRef<SLayeredImage> IconWidget =
@@ -331,6 +327,7 @@ void SToolBarButtonBlock::BuildMultiBlockWidget(const ISlateStyle* StyleSet, con
 	EMultiBlockLocation::Type BlockLocation = GetMultiBlockLocation();
 	
 	// What type of UI should we create for this block?
+	TWeakPtr<const FUICommandInfo> Action = ToolBarButtonBlock->GetAction();
 	EUserInterfaceActionType UserInterfaceType = ToolBarButtonBlock->UserInterfaceActionType;
 	if ( Action.IsValid() )
 	{
@@ -356,7 +353,7 @@ void SToolBarButtonBlock::BuildMultiBlockWidget(const ISlateStyle* StyleSet, con
 			.ButtonStyle(ToolbarButtonStyle)
 			.IsEnabled(this, &SToolBarButtonBlock::IsEnabled)
 			.OnClicked(this, &SToolBarButtonBlock::OnClicked)
-			.ToolTip(FMultiBoxSettings::ToolTipConstructor.Execute(ActualToolTip, nullptr, Action.Pin()))
+			.ToolTip(FMultiBoxSettings::ToolTipConstructor.Execute(ActualToolTip, nullptr, Action.Pin(), /*ShowActionShortcut=*/ true))
 			.IsFocusable(bIsFocusable)
 			[
 				ButtonContent
@@ -379,7 +376,7 @@ void SToolBarButtonBlock::BuildMultiBlockWidget(const ISlateStyle* StyleSet, con
 						.Style(CheckStyle)
 						.CheckBoxContentUsesAutoWidth(false)
 						.IsFocusable(bIsFocusable)
-						.ToolTip( FMultiBoxSettings::ToolTipConstructor.Execute( ActualToolTip, nullptr, Action.Pin()))		
+						.ToolTip( FMultiBoxSettings::ToolTipConstructor.Execute( ActualToolTip, nullptr, Action.Pin(), /*ShowActionShortcut=*/ true))		
 						.OnCheckStateChanged(this, &SToolBarButtonBlock::OnCheckStateChanged )
 						.OnGetMenuContent( ToolBarButtonBlock->OnGetMenuContent )
 						.IsChecked(this, &SToolBarButtonBlock::GetCheckState)
diff --git a/Engine/Source/Runtime/Slate/Private/Framework/MultiBox/SWidgetBlock.cpp b/Engine/Source/Runtime/Slate/Private/Framework/MultiBox/SWidgetBlock.cpp
--- a/Engine/Source/Runtime/Slate/Private/Framework/MultiBox/SWidgetBlock.cpp
+++ b/Engine/Source/Runtime/Slate/Private/Framework/MultiBox/SWidgetBlock.cpp
@@ -154,7 +154,7 @@ void SWidgetBlock::BuildMultiBlockWidget(const ISlateStyle* StyleSet, const FNam
 	.Padding( Padding )	// Large left margin mimics the indent of normal menu items when bNoIndent is false
 	[
 		SNew(SHorizontalBox)
-		.ToolTip(FMultiBoxSettings::ToolTipConstructor.Execute(WidgetBlock->ToolTipText, nullptr, nullptr ))
+		.ToolTip(FMultiBoxSettings::ToolTipConstructor.Execute(WidgetBlock->ToolTipText, nullptr, nullptr, /*ShowActionShortcut=*/ false ))
 		+ SHorizontalBox::Slot()
 		.AutoWidth()
 		[
diff --git a/Engine/Source/Runtime/Slate/Public/Framework/MultiBox/MultiBoxDefs.h b/Engine/Source/Runtime/Slate/Public/Framework/MultiBox/MultiBoxDefs.h
--- a/Engine/Source/Runtime/Slate/Public/Framework/MultiBox/MultiBoxDefs.h
+++ b/Engine/Source/Runtime/Slate/Public/Framework/MultiBox/MultiBoxDefs.h
@@ -65,7 +65,7 @@ class FMultiBoxSettings
 {
 public:
 
-	DECLARE_DELEGATE_RetVal_ThreeParams( TSharedRef< SToolTip >, FConstructToolTip, const TAttribute<FText>& /*ToolTipText*/, const TSharedPtr<SWidget>& /*OverrideContent*/, const TSharedPtr<const FUICommandInfo>& /*Action*/ );
+	DECLARE_DELEGATE_RetVal_FourParams( TSharedRef< SToolTip >, FConstructToolTip, const TAttribute<FText>& /*ToolTipText*/, const TSharedPtr<SWidget>& /*OverrideContent*/, const TSharedPtr<const FUICommandInfo>& /*Action*/, bool /*ShowActionShortcut*/);
 
 	/** Access to whether multiboxes use small icons or default sized icons */
 	static SLATE_API TAttribute<bool> UseSmallToolBarIcons;
@@ -76,7 +76,7 @@ public:
 
 	SLATE_API FMultiBoxSettings();
 
-	static SLATE_API TSharedRef< SToolTip > ConstructDefaultToolTip( const TAttribute<FText>& ToolTipText, const TSharedPtr<SWidget>& OverrideContent, const TSharedPtr<const FUICommandInfo>& Action );
+	static SLATE_API TSharedRef< SToolTip > ConstructDefaultToolTip( const TAttribute<FText>& ToolTipText, const TSharedPtr<SWidget>& OverrideContent, const TSharedPtr<const FUICommandInfo>& Action, bool ShowActionShortcut = true );
 
 	static SLATE_API void ResetToolTipConstructor();
 };
