Ensure tooltips finish their size animation

diff --git a/Engine/Source/Editor/Documentation/Private/SDocumentationToolTip.cpp b/Engine/Source/Editor/Documentation/Private/SDocumentationToolTip.cpp
--- a/Engine/Source/Editor/Documentation/Private/SDocumentationToolTip.cpp
+++ b/Engine/Source/Editor/Documentation/Private/SDocumentationToolTip.cpp
@@ -541,19 +541,18 @@ void SDocumentationToolTip::Tick( const FGeometry& AllottedGeometry, const doubl
 	const bool NeedsUpdate = IsDisplayingDocumentationLink != GetDefault<UEditorPerProjectUserSettings>()->bDisplayDocumentationLink;
 	if (TransitionStartTime > 0)
 	{
-		bIsInTransition = InCurrentTime - TransitionStartTime <= TransitionLength;
-		if (bIsInTransition) {
-			TransitionPercentage = (InCurrentTime - TransitionStartTime) / TransitionLength;
-			FVector2D TransitionEndSize = WidgetContent->GetDesiredSize();
-			if (TransitionEndSize.Y > TransitionStartSize.Y)
-			{
-				LastDesiredSize = ((TransitionEndSize - TransitionStartSize) * FMath::InterpEaseOut<float>(0.f, 1.f, TransitionPercentage, 4.f)) + TransitionStartSize;
-			}
-			else
-			{
-				LastDesiredSize = TransitionStartSize - ((TransitionStartSize - TransitionEndSize) * FMath::InterpEaseOut<float>(0.f, 1.f, TransitionPercentage, 4.f));
-			}
+		TransitionPercentage = (InCurrentTime - TransitionStartTime) / TransitionLength;
+		if (TransitionPercentage > 1.0f)
+		{
+			TransitionPercentage = 1.0f;
+			// Stop transition.
+			TransitionStartTime = 0;
 		}
+
+		const FVector2D TransitionEndSize = WidgetContent->GetDesiredSize();
+		LastDesiredSize = TransitionStartSize
+						- (((TransitionStartSize - TransitionEndSize))
+						   * FMath::InterpEaseOut<float>(0.f, 1.f, TransitionPercentage, 4.f));
 	}
 
 	if ( !IsShowingFullTip && ModifierKeys.IsAltDown() && ModifierKeys.IsControlDown() )
