Iterate on enhanced tooltip styles and animation.
- Dark theme for enhanced content.
- Narrower max width.
- Images scale to fit max width.
- Scrollbox for very long content excerpts.
- Allow three levels of headers.
- Move animation calculations into GetDesiredSize to minimize rendering artifacts.
- Accept override content to replace the "Hold Ctrl + Alt" prompt, and use it in the Content Browser.
- Animate the height of the prompt during the transition to avoid visual pops.

diff --git a/Engine/Source/Editor/Documentation/Private/DocumentationStyleSet.cpp b/Engine/Source/Editor/Documentation/Private/DocumentationStyleSet.cpp
--- a/Engine/Source/Editor/Documentation/Private/DocumentationStyleSet.cpp
+++ b/Engine/Source/Editor/Documentation/Private/DocumentationStyleSet.cpp
@@ -38,13 +38,29 @@ FDocumentationStyleSet::FDocumentationStyleSet()
 	const FVector2D Icon32x32(32.0f, 32.0f);
 	const FVector2D Icon128x128(128.0f, 128.0f);
 
-	Set("ToolTip.Header", new FSlateColorBrush(FStyleColors::Foreground));
+	const FLinearColor DarkHeaderBackground = COLOR("#383838FF");        // Dropdown
+	const FLinearColor ContentBackground = COLOR("#4F4F4FFF");           // New shade
+	const FLinearColor ItemBackground = COLOR("#808080FF");              // Hover2
+	const FLinearColor ItemBackgroundHighlight = COLOR("#A0A0A0FF");     // New shade between Hover2 and Foreground
 
-	Set("ToolTip.ScrollBox", FScrollBoxStyle());
+	Set("ToolTip.Header", new FSlateColorBrush(DarkHeaderBackground));
 
-	Set("ToolTip.TopSeparator", new BORDER_BRUSH("Common/Selector", FMargin(0.0f, 1.0f, 0.0f, 0.0f), FStyleColors::Hover));
+	Set("ToolTip.ScrollBox", FScrollBoxStyle()
+		.SetTopShadowBrush(FSlateNoResource())
+		.SetBottomShadowBrush(FSlateNoResource())
+	);
+	Set("ToolTip.Scrollbar", FScrollBarStyle(FAppStyle::GetWidgetStyle<FScrollBarStyle>("ScrollBar"))
+		.SetNormalThumbImage(FSlateRoundedBoxBrush(ItemBackground, 4.0f))
+		.SetDraggedThumbImage(FSlateRoundedBoxBrush(ItemBackgroundHighlight, 4.0f))
+		.SetHoveredThumbImage(FSlateRoundedBoxBrush(ItemBackgroundHighlight, 4.0f))
+		.SetVerticalBackgroundImage(FSlateColorBrush(DarkHeaderBackground))
+		.SetHorizontalBackgroundImage(FSlateColorBrush(DarkHeaderBackground))
+	);
+	Set("ToolTip.ContentBackground", new FSlateColorBrush(ContentBackground));
 
-	Set("ToolTip.KeybindBorder", new FSlateRoundedBoxBrush(FStyleColors::Hover2, 4.0f));
+	Set("ToolTip.TopSeparator", new BORDER_BRUSH("Common/Selector", FMargin(0.0f, 1.0f, 0.0f, 0.0f), FStyleColors::Black));
+
+	Set("ToolTip.KeybindBorder", new FSlateRoundedBoxBrush(ItemBackground, 4.0f));
 	const FTextBlockStyle KeybindText = FTextBlockStyle(FAppStyle::GetWidgetStyle<FTextBlockStyle>("NormalText"))
 		.SetFont(DEFAULT_FONT("Bold", 10))
 		.SetColorAndOpacity(FLinearColor(1.0f, 1.0f, 1.0f));
diff --git a/Engine/Source/Editor/Documentation/Private/SDocumentationToolTip.cpp b/Engine/Source/Editor/Documentation/Private/SDocumentationToolTip.cpp
--- a/Engine/Source/Editor/Documentation/Private/SDocumentationToolTip.cpp
+++ b/Engine/Source/Editor/Documentation/Private/SDocumentationToolTip.cpp
@@ -9,6 +9,7 @@
 #include "Widgets/Text/STextBlock.h"
 #include "Widgets/Layout/SBox.h"
 #include "Widgets/Layout/SScrollBox.h"
+#include "Widgets/Images/SImage.h"
 #include "Widgets/SToolTip.h"
 #include "Styling/AppStyle.h"
 #include "Editor/EditorPerProjectUserSettings.h"
@@ -25,7 +26,22 @@
 #include "Framework/Notifications/NotificationManager.h"
 #include "Widgets/Notifications/SNotificationList.h"
 #include "DocumentationStyleSet.h"
-#include "SPrimaryButton.h"
+
+
+namespace UE::UnrealEd::Private
+{
+	float CVarTooltipExpandTimeValue = 0.12f;
+
+	// The visual width of the outside border around the tooltip contents.
+	static float WidgetContentPadding = 1.5f;
+}
+
+static FAutoConsoleVariableRef CVarTooltipExpandTime(
+	TEXT("TooltipExpandTime"),
+	UE::UnrealEd::Private::CVarTooltipExpandTimeValue,
+	TEXT("Set the transition time for tooltip expansion."),
+	ECVF_Default
+);
 
 void SDocumentationToolTip::Construct( const FArguments& InArgs )
 {
@@ -42,6 +58,7 @@ void SDocumentationToolTip::Construct( const FArguments& InArgs )
 	IsDisplayingDocumentationLink = false;
 	Shortcut = InArgs._Shortcut;
 	OverrideFullTooltipContent = InArgs._OverrideExtendedToolTipContent;
+	OverridePromptContent = InArgs._OverridePromptContent;
 
 	ExcerptName = InArgs._ExcerptName;
 	IsShowingFullTip = false;
@@ -53,7 +70,8 @@ void SDocumentationToolTip::Construct( const FArguments& InArgs )
 		OverrideContent = InArgs._Content.Widget;
 	}
 
-	SAssignNew(DocumentationControlBox, SHorizontalBox);
+	SAssignNew(SimpleTipContent, SBox);
+	SAssignNew(DocumentationControls, SHorizontalBox);
 	SAssignNew(FullTipContent, SBox);
 	if (OverrideFullTooltipContent.IsValid())
 	{
@@ -62,15 +80,15 @@ void SDocumentationToolTip::Construct( const FArguments& InArgs )
 	}
 
 	ConstructSimpleTipContent();
-
 	ChildSlot
 	[
 		SAssignNew(WidgetContent, SBox)
-		.Padding(2.0f)
+		.Padding(UE::UnrealEd::Private::WidgetContentPadding)
 		[
 			SNew(SVerticalBox)
 			+SVerticalBox::Slot()
 			.AutoHeight()
+			.VAlign(VAlign_Top)
 			[
 				SimpleTipContent.ToSharedRef()
 			]
@@ -82,29 +100,18 @@ void SDocumentationToolTip::Construct( const FArguments& InArgs )
 			+ SVerticalBox::Slot()
 			.AutoHeight()
 			[
-				SNew(SBox)
+				SAssignNew(DocumentationControlBox, SBox)
 				.Visibility(this, &SDocumentationToolTip::GetControlVisibility)
 				[
-					SNew(SVerticalBox)
-					+ SVerticalBox::Slot()
-					.AutoHeight()
-					[
-						SNew(SBorder)
-						.Padding(0.0)
-						.BorderImage(FDocumentationStyleSet::Get().GetBrush("ToolTip.TopSeparator"))
-						[
-							SNew(SBox)
-							.HeightOverride(1.0f)
-						]
-					]
-					+ SVerticalBox::Slot()
-					.AutoHeight()
+					SNew(SBorder)
+					.Padding(0,1,0,0)
+					.BorderImage(FDocumentationStyleSet::Get().GetBrush("ToolTip.TopSeparator"))
 					[
 						SNew(SBorder)
-						.Padding(9.0)
+						.Padding(9)
 						.BorderImage(FDocumentationStyleSet::Get().GetBrush("ToolTip.Header"))
 						[
-							DocumentationControlBox.ToSharedRef()
+							DocumentationControls.ToSharedRef()
 						]
 					]
 				]
@@ -112,12 +119,9 @@ void SDocumentationToolTip::Construct( const FArguments& InArgs )
 		]
 	];
 
-	bIsInTransition = false;
 	TransitionStartTime = 0;
-	TransitionLength = 0.2f;
-	LastDesiredSize = SimpleTipContent->GetDesiredSize();
-	TransitionStartSize = SimpleTipContent->GetDesiredSize();
 	bFullTipContentIsReady = OverrideFullTooltipContent.IsValid();
+	bIsPromptVisible = !IsShowingFullTip;
 }
 
 void SDocumentationToolTip::ConstructSimpleTipContent()
@@ -150,11 +154,10 @@ void SDocumentationToolTip::ConstructSimpleTipContent()
 	TSharedPtr< SHorizontalBox > TextBox;
 	if ( !OverrideContent.IsValid() )
 	{
-		SAssignNew( SimpleTipContent, SBox )
-		[
+		SimpleTipContent->SetContent(
 			SNew(SBorder)
-			.BorderImage(this, &SDocumentationToolTip::GetSimpleTipBorderStyle)
 			.Padding(9.f)
+			.BorderImage(FAppStyle::GetNoBrush())
 			[
 				SAssignNew( VerticalBox, SVerticalBox )
 				+SVerticalBox::Slot()
@@ -174,7 +177,7 @@ void SDocumentationToolTip::ConstructSimpleTipContent()
 					]
 				]
 			]
-		];
+		);
 
 		TextBox->AddSlot()
 		.VAlign(VAlign_Center)
@@ -198,11 +201,10 @@ void SDocumentationToolTip::ConstructSimpleTipContent()
 	}
 	else
 	{
-		SAssignNew( SimpleTipContent, SBox )
-		[
+		SimpleTipContent->SetContent(
 			SNew(SBorder)
-			.BorderImage(this, &SDocumentationToolTip::GetSimpleTipBorderStyle)
 			.Padding(9.f)
+			.BorderImage(FAppStyle::GetNoBrush())
 			[
 				SAssignNew( VerticalBox, SVerticalBox )
 				+SVerticalBox::Slot()
@@ -211,7 +213,7 @@ void SDocumentationToolTip::ConstructSimpleTipContent()
 					OverrideContent.ToSharedRef()
 				]
 			]
-		];
+		);
 	}
 
 	if (bAddDocumentation)
@@ -222,62 +224,95 @@ void SDocumentationToolTip::ConstructSimpleTipContent()
 
 void SDocumentationToolTip::AddDocumentation(TSharedPtr< SVerticalBox > VerticalBox)
 {
-	if ( !DocumentationLink.IsEmpty() && !OverrideFullTooltipContent.IsValid())
+	if ( !DocumentationLink.IsEmpty() || OverrideFullTooltipContent.IsValid() )
 	{
 		if ( !DocumentationPage.IsValid() )
 		{
 			DocumentationPage = IDocumentation::Get()->GetPage( DocumentationLink, NULL );
 		}
 
-		if ( DocumentationPage->HasExcerpt( ExcerptName ))
+		if ( DocumentationPage->HasExcerpt( ExcerptName ) || OverrideFullTooltipContent.IsValid())
 		{
-			FText MacShortcut = NSLOCTEXT("SToolTip", "MacRichTooltipShortcut", "Command + Option");
-			FText WinShortcut = NSLOCTEXT("SToolTip", "WinRichTooltipShortcut", "Ctrl + Alt");
+			if (OverridePromptContent.IsValid())
+			{
+				VerticalBox->AddSlot()
+					.AutoHeight()
+					.HAlign(HAlign_Center)
+					[
+						SNew(SBox)
+						.Visibility(this, &SDocumentationToolTip::GetPromptVisibility)
+						.HeightOverride(this, &SDocumentationToolTip::GetPromptHeight)
+						[
+							SAssignNew(PromptContent,SBox)
+							[
+								OverridePromptContent.ToSharedRef()
+							]
+						]
+					];
+			}
+			else
+			{
+				FText MacShortcut = NSLOCTEXT("SToolTip", "MacRichTooltipShortcut", "Command + Option");
+				FText WinShortcut = NSLOCTEXT("SToolTip", "WinRichTooltipShortcut", "Ctrl + Alt");
 
-			FText KeyboardShortcut;
+				FText KeyboardShortcut;
 #if PLATFORM_MAC
-			KeyboardShortcut = MacShortcut;
+				KeyboardShortcut = MacShortcut;
 #else
-			KeyboardShortcut = WinShortcut;
+				KeyboardShortcut = WinShortcut;
 #endif
 
-			VerticalBox->AddSlot()
-			.AutoHeight()
-			.HAlign( HAlign_Right )
-			[
-				SNew(SBox)
-				.Visibility(this, &SDocumentationToolTip::GetPromptVisibility)
-				.Padding(0, 9, 0, 0)
+				VerticalBox->AddSlot()
+				.AutoHeight()
+				.HAlign( HAlign_Center )
 				[
-					SNew(SHorizontalBox)
-					+SHorizontalBox::Slot()
-					.VAlign(VAlign_Center)
-					.AutoWidth()
-					.Padding(0.0f, 0.0f, 5.0f, 0.0f)
-					[
-						SNew( STextBlock )
-						.TextStyle( &SubduedStyleInfo )
-						.Text( NSLOCTEXT( "SToolTip", "AdvancedToolTipMessage", "Learn more: hold" ) )
-					]
-					+SHorizontalBox::Slot()
-					.VAlign(VAlign_Center)
-					.AutoWidth()
+					SNew(SBox)
+					.Visibility(this, &SDocumentationToolTip::GetPromptVisibility)
+					.HeightOverride(this, &SDocumentationToolTip::GetPromptHeight)
 					[
-						SNew( SBorder )
-						.BorderImage(FDocumentationStyleSet::Get().GetBrush("ToolTip.ToggleKeybindBorder"))
-						.Padding(4.0f, 2.0f)
+						SAssignNew(PromptContent, SBox)
 						[
-							SNew(STextBlock)
-							.TextStyle(&SubduedStyleInfo)
-							.Text(FText::Format(NSLOCTEXT("SToolTip", "AdvancedToolTipKeybind", "{0}"), KeyboardShortcut))
-							.Visibility(this, &SDocumentationToolTip::GetPromptVisibility)
+							SNew(SHorizontalBox)
+							+SHorizontalBox::Slot()
+							.Padding(0, 5, 0, 0)
+							.VAlign(VAlign_Center)
+							.HAlign(HAlign_Center)
+							.AutoWidth()
+							[
+								SNew( STextBlock )
+								.TextStyle( &SubduedStyleInfo )
+								.Text( NSLOCTEXT( "SToolTip", "AdvancedToolTipMessage", "Hold" ) )
+							]
+							+SHorizontalBox::Slot()
+							.VAlign(VAlign_Center)
+							.AutoWidth()
+							.Padding(5, 5, 5, 0)
+							[
+								SNew( SBorder )
+								.BorderImage(FDocumentationStyleSet::Get().GetBrush("ToolTip.ToggleKeybindBorder"))
+								.Padding(4, 2)
+								[
+									SNew(STextBlock)
+									.TextStyle(&SubduedStyleInfo)
+									.Text(FText::Format(NSLOCTEXT("SToolTip", "AdvancedToolTipKeybind", "{0}"), KeyboardShortcut))
+									.Visibility(this, &SDocumentationToolTip::GetPromptVisibility)
+								]
+							]
+							+ SHorizontalBox::Slot()
+							.VAlign(VAlign_Center)
+							.Padding(0, 5, 0, 0)
+							.AutoWidth()
+							[
+								SNew(STextBlock)
+									.TextStyle(&SubduedStyleInfo)
+									.Text(NSLOCTEXT("SToolTip", "AdvancedToolTipMessageEnd", "for more"))
+							]
 						]
 					]
-				]
-			];
+				];
+			}
 		}
 
-		SAssignNew(DocumentationControlBox, SHorizontalBox);
 		IsDisplayingDocumentationLink = GetDefault<UEditorPerProjectUserSettings>()->bDisplayDocumentationLink;
 		if (IsDisplayingDocumentationLink)
 		{
@@ -287,14 +322,14 @@ void SDocumentationToolTip::AddDocumentation(TSharedPtr< SVerticalBox > Vertical
 			{
 				OptionalExcerptName = FString(TEXT(" [")) + ExcerptName + TEXT("]");
 			}
-			DocumentationControlBox->AddSlot()
+			DocumentationControls->AddSlot()
 			.FillWidth(1.0f)
 			.HAlign(HAlign_Left)
 			.Padding(0, 0, 9, 0)
 			[
 				SNew(STextBlock)
 					.Text(FText::FromString(DocumentationLink + OptionalExcerptName))
-					.TextStyle(&SubduedStyleInfo)
+					.TextStyle(FAppStyle::Get(), "Documentation.SDocumentationTooltip")
 			];
 
 			if (!DocumentationPage->HasExcerpt(ExcerptName) && FSlateApplication::Get().SupportsSourceAccess())
@@ -305,7 +340,7 @@ void SDocumentationToolTip::AddDocumentation(TSharedPtr< SVerticalBox > Vertical
 					DocPath = FPaths::ConvertRelativePathToFull(DocPath);
 				}
 
-				DocumentationControlBox->AddSlot()
+				DocumentationControls->AddSlot()
 				.AutoWidth()
 				.HAlign(HAlign_Right)
 				[
@@ -422,24 +457,36 @@ void SDocumentationToolTip::ConstructFullTipContent()
 		{
 			TSharedPtr< SVerticalBox > Box;
 			TSharedPtr< SWidget > FullTipBox = SNew(SBox)
-			.Visibility(this, &SDocumentationToolTip::GetFullTipVisibility)
 			.Padding(DocumentationMargin)
+			.Visibility(this, &SDocumentationToolTip::GetFullTipVisibility)
 			[
-				SAssignNew(Box, SVerticalBox)
-				+ SVerticalBox::Slot()
-				.HAlign(HAlign_Center)
-				.AutoHeight()
-				.MaxHeight(900.f)
+				SNew(SBorder)
+				.BorderImage(FDocumentationStyleSet::Get().GetBrush("ToolTip.TopSeparator"))
+				.Padding(0,1,0,0)
 				[
 					SNew(SBorder)
-					.BorderImage(FDocumentationStyleSet::Get().GetBrush("ToolTip.TopSeparator"))
+					.BorderImage(FDocumentationStyleSet::Get().GetBrush("ToolTip.ContentBackground"))
+					.Padding(0)
 					[
-						SNew(SScrollBox)
-						.Style(FDocumentationStyleSet::Get(), "ToolTip.ScrollBox")
-						+SScrollBox::Slot()
-						.Padding(9.0f)
+						SAssignNew(Box, SVerticalBox)
+						+ SVerticalBox::Slot()
+						.HAlign(HAlign_Center)
+						.MaxHeight(750)
 						[
-							Excerpts[ExcerptIndex].Content.ToSharedRef()
+							SNew(SScrollBox)
+							.Style(FDocumentationStyleSet::Get(), "ToolTip.ScrollBox")
+							.ScrollBarStyle(FDocumentationStyleSet::Get(), "ToolTip.ScrollBar")
+							+SScrollBox::Slot()
+							.Padding(14.0f)
+							.FillSize(1.0f)
+							[
+								SNew(SHorizontalBox)
+								+ SHorizontalBox::Slot()
+								.FillWidth(1.0f)
+								[
+									Excerpts[ExcerptIndex].Content.ToSharedRef()
+								]
+							]
 						]
 					]
 				]
@@ -455,18 +502,6 @@ void SDocumentationToolTip::ConstructFullTipContent()
 					BaseUrl = *ExcerptBaseUrl;
 				}
 
-				Box->AddSlot()
-				.AutoHeight()
-				[
-					SNew(SBorder)
-					.BorderImage(FDocumentationStyleSet::Get().GetBrush("ToolTip.TopSeparator"))
-					.Padding(0.0f)
-					[
-						SNew(SBox)
-						.HeightOverride(1.0f)
-					]
-				];
-
 				Box->AddSlot()
 				.AutoHeight()
 				[
@@ -478,10 +513,32 @@ void SDocumentationToolTip::ConstructFullTipContent()
 						+SHorizontalBox::Slot()
 						.FillWidth(1.0f)
 						.HAlign(HAlign_Right)
+						.VAlign(VAlign_Center)
 						[
-							SNew(SPrimaryButton)
-							.Icon(FAppStyle::Get().GetBrush("Icons.Help"))
-							.Text(NSLOCTEXT("SToolTip", "LearnMoreButton", "Learn More Online"))
+							SNew(SButton)
+							.ButtonStyle(&FAppStyle::Get().GetWidgetStyle<FButtonStyle>("PrimaryButton"))
+							.ContentPadding(FMargin(0.0f, 4.5f))
+							.Content()
+							[
+								SNew(SHorizontalBox)
+								+SHorizontalBox::Slot()
+								.VAlign(VAlign_Center)
+								.AutoWidth()
+								.Padding(0,0,3,0)
+								[
+									SNew(SImage)
+									.Image(FAppStyle::GetBrush("Icons.Help.Solid"))
+								]
+								+SHorizontalBox::Slot()
+								.VAlign(VAlign_Center)
+								.AutoWidth()
+								[
+									SNew(STextBlock)
+									.TextStyle(&FAppStyle::Get().GetWidgetStyle<FTextBlockStyle>("PrimaryButtonText"))
+									.Text(NSLOCTEXT("SToolTip", "LearnMoreButton", "Learn More Online"))
+								]
+							]
+
 							.OnClicked_Static([](FString Link, FString BaseUrl) -> FReply {
 								if (!IDocumentation::Get()->Open(Link, FDocumentationSourceInfo(TEXT("rich_tooltips")), BaseUrl))
 								{
@@ -497,7 +554,7 @@ void SDocumentationToolTip::ConstructFullTipContent()
 
 			if (IsDisplayingDocumentationLink && FSlateApplication::Get().SupportsSourceAccess() )
 			{
-				DocumentationControlBox->AddSlot()
+				DocumentationControls->AddSlot()
 				.AutoWidth()
 				.HAlign( HAlign_Right )
 				[
@@ -521,10 +578,9 @@ void SDocumentationToolTip::ConstructFullTipContent()
 
 FReply SDocumentationToolTip::ReloadDocumentation()
 {
-	SimpleTipContent.Reset();
-	DocumentationControlBox.Reset();
 	bFullTipContentIsReady = false;
 
+	DocumentationControls->ClearChildren();
 	ConstructSimpleTipContent();
 
 	if ( DocumentationPage.IsValid() )
@@ -542,24 +598,16 @@ FReply SDocumentationToolTip::ReloadDocumentation()
 
 void SDocumentationToolTip::Tick( const FGeometry& AllottedGeometry, const double InCurrentTime, const float InDeltaTime )
 {
-	LastDesiredSize = WidgetContent->GetDesiredSize();
-
 	const FModifierKeysState ModifierKeys = FSlateApplication::Get().GetModifierKeys();
 	const bool NeedsUpdate = !OverrideFullTooltipContent.IsValid() && IsDisplayingDocumentationLink != GetDefault<UEditorPerProjectUserSettings>()->bDisplayDocumentationLink;
 	if (TransitionStartTime > 0)
 	{
-		TransitionPercentage = (InCurrentTime - TransitionStartTime) / TransitionLength;
-		if (TransitionPercentage > 1.0f)
+		const float TransitionPercentage = (InCurrentTime - TransitionStartTime) / TransitionLength;
+		if (TransitionPercentage >= 1.0f)
 		{
-			TransitionPercentage = 1.0f;
 			// Stop transition.
 			TransitionStartTime = 0;
 		}
-
-		const FVector2D TransitionEndSize = WidgetContent->GetDesiredSize();
-		LastDesiredSize = TransitionStartSize
-						- (((TransitionStartSize - TransitionEndSize))
-						   * FMath::InterpEaseOut<float>(0.f, 1.f, TransitionPercentage, 4.f));
 	}
 
 	if ( !IsShowingFullTip && ModifierKeys.IsAltDown() && ModifierKeys.IsControlDown() )
@@ -590,24 +638,22 @@ void SDocumentationToolTip::Tick( const FGeometry& AllottedGeometry, const doubl
 					FEngineAnalytics::GetProvider().RecordEvent(TEXT("Editor.Usage.Documentation.FullTooltipShown"), Params);
 				}
 			}
-			bIsInTransition = true;
-			IsShowingFullTip = true;
+			bIsPromptVisible = false;
 			TransitionStartTime = InCurrentTime;
-			TransitionStartSize = LastDesiredSize;
+			TransitionLength = UE::UnrealEd::Private::CVarTooltipExpandTimeValue;
 		}
+		IsShowingFullTip = true;
 	}
 	else if ( ( IsShowingFullTip || NeedsUpdate )  && ( !ModifierKeys.IsAltDown() || !ModifierKeys.IsControlDown() ) )
 	{
 		if ( NeedsUpdate )
 		{
-			ReloadDocumentation();
 			IsDisplayingDocumentationLink = GetDefault<UEditorPerProjectUserSettings>()->bDisplayDocumentationLink;
+			ReloadDocumentation();
 		}
 
-		bIsInTransition = true;
 		IsShowingFullTip = false;
-		TransitionStartTime = InCurrentTime;
-		TransitionStartSize = LastDesiredSize;
+		bIsPromptVisible = true;
 	}
 }
 
@@ -619,7 +665,24 @@ bool SDocumentationToolTip::IsInteractive() const
 
 FVector2D SDocumentationToolTip::ComputeDesiredSize(float LayoutScaleMultiplier) const
 {
-	return LastDesiredSize;
+	FVector2D SimpleTipSize = SimpleTipContent->GetDesiredSize();
+	FVector2D TransitionStartSize = SimpleTipSize + (2 * UE::UnrealEd::Private::WidgetContentPadding);
+	if (GetControlVisibility() == EVisibility::Visible)
+	{
+		TransitionStartSize.Y += DocumentationControlBox->GetDesiredSize().Y;
+		TransitionStartSize.X = FMath::Max(TransitionStartSize.X, DocumentationControlBox->GetDesiredSize().X + (2 * UE::UnrealEd::Private::WidgetContentPadding));
+	}
+	if (TransitionStartTime > 0 && !IsDisplayingDocumentationLink)
+	{
+		const float TransitionPercentage = FMath::Clamp(((FSlateApplication::Get().GetCurrentTime() - TransitionStartTime) / TransitionLength), 0, 1);
+
+		const FVector2D TransitionEndSize = (IsShowingFullTip) ? WidgetContent->GetDesiredSize() : TransitionStartSize;
+		return (TransitionStartSize
+			- ((TransitionStartSize - TransitionEndSize)
+				* FMath::InterpEaseOut<float>(0.f, 1.f, TransitionPercentage, 4.f)));
+	}
+
+	return (IsShowingFullTip) ? WidgetContent->GetDesiredSize() : TransitionStartSize;
 }
 
 EVisibility SDocumentationToolTip::GetOverriddenFullToolTipVisibility() const
@@ -627,22 +690,14 @@ EVisibility SDocumentationToolTip::GetOverriddenFullToolTipVisibility() const
 	return IsShowingFullTip ? EVisibility::Visible : EVisibility::Collapsed;
 }
 
-EVisibility SDocumentationToolTip::GetFullTipVisibility() const
+EVisibility SDocumentationToolTip::GetPromptVisibility() const
 {
-	if (IsShowingFullTip)
-	{
-		return EVisibility::Visible;
-	}
-	return EVisibility::Collapsed;
+	return (bIsPromptVisible) ? EVisibility::Visible : EVisibility::Hidden;
 }
 
-EVisibility SDocumentationToolTip::GetPromptVisibility() const
+EVisibility SDocumentationToolTip::GetFullTipVisibility() const
 {
-	if (IsShowingFullTip)
-	{
-		return EVisibility::Collapsed;
-	}
-	return EVisibility::Visible;
+	return (IsShowingFullTip) ? EVisibility::Visible : EVisibility::Collapsed;
 }
 
 EVisibility SDocumentationToolTip::GetControlVisibility() const
@@ -663,11 +718,18 @@ EVisibility SDocumentationToolTip::GetShortcutVisibility() const
 	return EVisibility::Collapsed;
 }
 
-const FSlateBrush* SDocumentationToolTip::GetSimpleTipBorderStyle() const
+FOptionalSize SDocumentationToolTip::GetPromptHeight() const
 {
-	if (IsShowingFullTip && !OverrideContent.IsValid())
+	const float DesiredHeight = PromptContent->GetDesiredSize().Y;
+	if (IsShowingFullTip)
 	{
-		return FDocumentationStyleSet::Get().GetBrush("ToolTip.Header");
+		if (TransitionStartTime > 0)
+		{
+			const float TransitionPercentage = FMath::Clamp(((FSlateApplication::Get().GetCurrentTime() - TransitionStartTime) / (0.75 * TransitionLength)), 0, 1);
+			return DesiredHeight * (1 - FMath::InterpEaseOut<float>(0.f, 1.f, TransitionPercentage, 4.f));
+		}
+		return 0.0f;
 	}
-	return FAppStyle::GetBrush("");
+	return DesiredHeight;
 }
+
diff --git a/Engine/Source/Editor/Documentation/Private/UDNParser.cpp b/Engine/Source/Editor/Documentation/Private/UDNParser.cpp
--- a/Engine/Source/Editor/Documentation/Private/UDNParser.cpp
+++ b/Engine/Source/Editor/Documentation/Private/UDNParser.cpp
@@ -53,6 +53,7 @@
 #include "Widgets/Notifications/SNotificationList.h"
 #include "Widgets/SBoxPanel.h"
 #include "Widgets/Text/STextBlock.h"
+#include "Styling/StyleColors.h"
 
 class SWidget;
 
@@ -103,8 +104,8 @@ TSharedRef< FUDNParser > FUDNParser::Create( const TSharedPtr< FParserConfigurat
 FUDNParser::FUDNParser( const TSharedRef< FParserConfiguration >& InConfiguration, const FDocumentationStyle& InStyle )
 	: Configuration( InConfiguration )
 	, Style(InStyle)
-	, WrapAt(600.f)
-	, ContentWidth(600.f)
+	, WrapAt(450.f)
+	, ContentWidth(450.f)
 {
 
 }
@@ -171,13 +172,17 @@ void FUDNParser::Initialize()
 	TokenArray.Add(EUDNToken::Pound);
 	TokenArray.Add(EUDNToken::Pound);
 	TokenArray.Add(EUDNToken::Pound);
-	LineLibrary.Add(FTokenConfiguration(TokenArray, FUDNLine::Header2, true));
+	LineLibrary.Add(FTokenConfiguration(TokenArray, FUDNLine::Header3, true));
 
 	TokenArray.Empty();
 	TokenArray.Add(EUDNToken::Pound);
 	TokenArray.Add(EUDNToken::Pound);
-	LineLibrary.Add(FTokenConfiguration(TokenArray, FUDNLine::Header1, true));
+	LineLibrary.Add(FTokenConfiguration(TokenArray, FUDNLine::Header2, true));
 	
+	TokenArray.Empty();
+	TokenArray.Add(EUDNToken::Pound);
+	LineLibrary.Add(FTokenConfiguration(TokenArray, FUDNLine::Header1, true));
+
 	TokenArray.Empty();
 	TokenArray.Add(EUDNToken::OpenBracket);
 	TokenArray.Add(EUDNToken::Content);
@@ -813,6 +818,7 @@ void FUDNParser::AddCaptionToExcerpt(TSharedPtr<SVerticalBox> Box, const FString
 					SNew(STextBlock)
 					.Text(FText::Format(LOCTEXT("ImageCaptionLabel", "Image: {0}"), FText::FromString(ContentSource)))
 					.TextStyle(FAppStyle::Get(), Style.ItalicContentStyleName)
+					.ColorAndOpacity(FStyleColors::Foreground)
 					.WrapTextAt(WrapAt.Get() - IndentWidth)
 				]
 			]
@@ -872,9 +878,6 @@ TSharedRef< SWidget > FUDNParser::GenerateExcerptContent( const FString& InLink,
 
 	const FString FullPath = FPaths::GetPath(InLink);
 
-	FSlateFontInfo Header1Font = FCoreStyle::GetDefaultFontStyle("Regular", 18);
-	FSlateFontInfo Header2Font = FCoreStyle::GetDefaultFontStyle("Regular", 14);
-
 	bool bCriticalError = false;
 	FString VariableName;
 	FString CurrentStringContent;
@@ -1015,14 +1018,15 @@ TSharedRef< SWidget > FUDNParser::GenerateExcerptContent( const FString& InLink,
 				CurrentStringContent.Empty();
 
 				Box->AddSlot()
-					.HAlign(HAlign_Center)
+					.Padding(0.0f, 8.0f, 0.0f, 16.0f)
 					[
-						SNew(SBox)
-						.WidthOverride(ContentWidth)
-						.Padding(FMargin(0,0,0,10))
+						SNew(SBorder)
+						.BorderImage(FAppStyle::Get().GetBrush("Documentation.Separator"))
+						.Padding(0.0f, 1.0f)
 						[
-							SNew(SSeparator)
-							.SeparatorImage(FAppStyle::GetBrush(Style.SeparatorStyleName))
+							SNew(SBox)
+							.WidthOverride(ContentWidth)
+							.HeightOverride(0)
 						]
 					];
 
@@ -1054,6 +1058,19 @@ TSharedRef< SWidget > FUDNParser::GenerateExcerptContent( const FString& InLink,
 				AddLineSeperator(Excerpt);
 				Excerpt.RichText += FString::Printf(TEXT("<TextStyle Style=\"%s\">%s</>"), *Style.Header2StyleName.ToString(), *Line.AdditionalContent[0]);
 				break;
+			case FUDNLine::Header3:
+				AddContentToExcerpt(Box, CurrentStringContent, Excerpt);
+				CurrentStringContent.Empty();
+
+				AppendExcerpt(Box,
+					SNew(STextBlock)
+					.Text(FText::FromString(Line.AdditionalContent[0]))
+					.TextStyle(FAppStyle::Get(), Style.Header3StyleName)
+				);
+
+				AddLineSeperator(Excerpt);
+				Excerpt.RichText += FString::Printf(TEXT("<TextStyle Style=\"%s\">%s</>"), *Style.Header3StyleName.ToString(), *Line.AdditionalContent[0]);
+				break;
 			case FUDNLine::Link:
 				AddContentToExcerpt(Box, CurrentStringContent, Excerpt);
 				CurrentStringContent.Empty();
@@ -1100,6 +1117,7 @@ TSharedRef< SWidget > FUDNParser::GenerateExcerptContent( const FString& InLink,
 						SNew( SImage )
 						.Image(DynamicBrush.Get())
 						.ToolTipText(FText::FromString(Line.AdditionalContent[0]))
+						.DesiredSizeOverride(GetImageDesiredSize(DynamicBrush, ContentWidth.Get().Get()))
 						);
 
 					DynamicBrushesUsed.AddUnique(DynamicBrush);
@@ -1490,4 +1508,14 @@ bool FUDNParser::ParseAssetLink(FString &InternalLink)
 	return false;
 }
 
+TOptional<FVector2D> FUDNParser::GetImageDesiredSize(TSharedPtr<FSlateDynamicImageBrush> DynamicBrush, float MaxWidth)
+{
+	if (MaxWidth && DynamicBrush->GetImageSize().X > MaxWidth)
+	{
+		float ratio = MaxWidth / DynamicBrush->GetImageSize().X;
+		return FVector2D(MaxWidth, ratio * DynamicBrush->GetImageSize().Y);
+	}
+	return FVector2D(DynamicBrush->GetImageSize());
+}
+
 #undef LOCTEXT_NAMESPACE
diff --git a/Engine/Source/Editor/Documentation/Private/UDNParser.h b/Engine/Source/Editor/Documentation/Private/UDNParser.h
--- a/Engine/Source/Editor/Documentation/Private/UDNParser.h
+++ b/Engine/Source/Editor/Documentation/Private/UDNParser.h
@@ -93,6 +93,7 @@ public:
 		NumberedContent,
 		Header1,
 		Header2,
+		Header3,
 		ExcerptOpen,
 		ExcerptClose,
 		Image,
@@ -222,6 +223,8 @@ private:
 	FUDNParser( const TSharedRef< FParserConfiguration >& InConfiguration, const FDocumentationStyle& InStyle );
 	void Initialize();
 
+	static TOptional<FVector2D> GetImageDesiredSize(TSharedPtr<FSlateDynamicImageBrush> DynamicBrush, float MaxWidth);
+
 	/** A list of dynamic brushes we are using for the currently loaded tutorial */
 	TArray< TSharedPtr< FSlateDynamicImageBrush > > DynamicBrushesUsed;
 
diff --git a/Engine/Source/Editor/Documentation/Public/IDocumentation.h b/Engine/Source/Editor/Documentation/Public/IDocumentation.h
--- a/Engine/Source/Editor/Documentation/Public/IDocumentation.h
+++ b/Engine/Source/Editor/Documentation/Public/IDocumentation.h
@@ -24,6 +24,7 @@ public:
 		, NumberedContentStyleName(TEXT("Documentation.NumberedContent"))
 		, Header1StyleName(TEXT("Documentation.Header1"))
 		, Header2StyleName(TEXT("Documentation.Header2"))
+		, Header3StyleName(TEXT("Documentation.Header3"))
 		, HyperlinkStyleName(TEXT("Documentation.Hyperlink"))
 		, HyperlinkButtonStyleName(TEXT("Documentation.Hyperlink.Button"))
 		, HyperlinkTextStyleName(TEXT("Documentation.Hyperlink.Text"))
@@ -73,6 +74,13 @@ public:
 		return *this;
 	}
 
+	/** Set the header 3 style for this documentation */
+	FDocumentationStyle& Header3Style(const FName& InName)
+	{
+		Header3StyleName = InName;
+		return *this;
+	}
+
 	/** Set the hyperlink button style for this documentation */
 	FDocumentationStyle& HyperlinkStyle(const FName& InName) 
 	{
@@ -119,6 +127,9 @@ public:
 	/** Header2 text style */
 	FName Header2StyleName;
 
+	/** Header3 text style */
+	FName Header3StyleName;
+
 	/** Hyperlink style */
 	FName HyperlinkStyleName;
 
diff --git a/Engine/Source/Editor/Documentation/Public/SDocumentationToolTip.h b/Engine/Source/Editor/Documentation/Public/SDocumentationToolTip.h
--- a/Engine/Source/Editor/Documentation/Public/SDocumentationToolTip.h
+++ b/Engine/Source/Editor/Documentation/Public/SDocumentationToolTip.h
@@ -39,6 +39,7 @@ public:
 		, _AddDocumentation( true )
 		, _DocumentationMargin(0)
 		, _OverrideExtendedToolTipContent(nullptr)
+		, _OverridePromptContent(nullptr)
 		, _Content()
 		{}
 
@@ -77,6 +78,9 @@ public:
 		/** Extended ToolTip to use when Ctrl + Alt is pressed instead of loading it from the Documentation */
 		SLATE_ARGUMENT(TSharedPtr<SWidget>, OverrideExtendedToolTipContent)
 
+		/** Custom widget to prompt users to hit Ctrl + Alt, if needed to fit with the design of custom content overrides */
+		SLATE_ARGUMENT(TSharedPtr<SWidget>, OverridePromptContent)
+
 		/** Arbitrary content to be displayed in the tool tip; overrides any text that may be set. */
 		SLATE_DEFAULT_SLOT( FArguments, Content )
 
@@ -120,11 +124,11 @@ private:
 
 	void CreateExcerpt( FString FileSource, FString ExcerptName );
 
-	EVisibility GetFullTipVisibility() const;
 	EVisibility GetPromptVisibility() const;
+	EVisibility GetFullTipVisibility() const;
+	FOptionalSize GetPromptHeight() const;
 	EVisibility GetControlVisibility() const;
 	EVisibility GetShortcutVisibility() const;
-	const FSlateBrush* GetSimpleTipBorderStyle() const;
 
 private:
 
@@ -133,6 +137,7 @@ private:
 	TAttribute< FText > Shortcut;
 	TSharedPtr< SWidget > OverrideContent;
 	TSharedPtr< SWidget > OverrideFullTooltipContent;
+	TSharedPtr< SWidget > OverridePromptContent;
 	FTextBlockStyle StyleInfo;
 	FTextBlockStyle SubduedStyleInfo;
 	FTextBlockStyle HyperlinkTextStyleInfo;
@@ -147,11 +152,13 @@ private:
 	/** Content widget */
 	TSharedPtr< SBox > WidgetContent;
 
-	TSharedPtr< SWidget > SimpleTipContent;
+	TSharedPtr< SBox > SimpleTipContent;
 	bool IsDisplayingDocumentationLink;
 
 	TSharedPtr< SBox > FullTipContent;
-	TSharedPtr< SHorizontalBox > DocumentationControlBox;
+	TSharedPtr< SBox > PromptContent;
+	TSharedPtr< SBox > DocumentationControlBox;
+	TSharedPtr< SHorizontalBox > DocumentationControls;
 
 	TSharedPtr< IDocumentationPage > DocumentationPage;
 	bool IsShowingFullTip;
@@ -159,12 +166,9 @@ private:
 	bool bAddDocumentation;
 	FMargin DocumentationMargin;
 
-	FVector2D LastDesiredSize;
-	bool bIsInTransition;
 	double TransitionStartTime;
 	float TransitionLength;
-	float TransitionPercentage;
-	FVector2D TransitionStartSize;
 
 	bool bFullTipContentIsReady;
+	bool bIsPromptVisible;
 };
diff --git a/Engine/Source/Editor/EditorStyle/Private/StarshipStyle.cpp b/Engine/Source/Editor/EditorStyle/Private/StarshipStyle.cpp
--- a/Engine/Source/Editor/EditorStyle/Private/StarshipStyle.cpp
+++ b/Engine/Source/Editor/EditorStyle/Private/StarshipStyle.cpp
@@ -2915,11 +2915,11 @@ void FStarshipEditorStyle::FStyle::SetupDockingStyles()
 void FStarshipEditorStyle::FStyle::SetupTutorialStyles()
 {
 	// Documentation tooltip defaults
-	const FSlateColor HyperlinkColor( FLinearColor( 0.1f, 0.1f, 0.5f ) );
+	const FSlateColor HyperlinkColor( FStyleColors::Foreground );
 	{
 		const FTextBlockStyle DocumentationTooltipText = FTextBlockStyle( NormalText )
 			.SetFont( DEFAULT_FONT( "Regular", 9 ) )
-			.SetColorAndOpacity( FLinearColor::Black );
+			.SetColorAndOpacity( FLinearColor::White );
 		Set("Documentation.SDocumentationTooltip", FTextBlockStyle(DocumentationTooltipText));
 
 		const FTextBlockStyle DocumentationTooltipTextSubdued = FTextBlockStyle( NormalText )
@@ -2942,7 +2942,7 @@ void FStarshipEditorStyle::FStyle::SetupTutorialStyles()
 
 	// Documentation defaults
 	const FTextBlockStyle DocumentationText = FTextBlockStyle(NormalText)
-		.SetColorAndOpacity( FLinearColor::Black )
+		.SetColorAndOpacity( FLinearColor::White )
 		.SetFont(DEFAULT_FONT( "Regular", 10 ));
 	Set("Documentation.Text", FTextBlockStyle(DocumentationText));
 
@@ -2951,8 +2951,8 @@ void FStarshipEditorStyle::FStyle::SetupTutorialStyles()
 	Set("Documentation.Hyperlink.Text", FTextBlockStyle(DocumentationHyperlinkText));
 
 	const FTextBlockStyle DocumentationHeaderText = FTextBlockStyle(NormalText)
-		.SetColorAndOpacity( FLinearColor::Black )
-		.SetFont(DEFAULT_FONT("Black", 32));
+		.SetColorAndOpacity( FLinearColor::White )
+		.SetFont(DEFAULT_FONT("Black", 16));
 	Set("Documentation.Header.Text", FTextBlockStyle(DocumentationHeaderText));
 
 	const FButtonStyle DocumentationHyperlinkButton = FButtonStyle()
@@ -2980,12 +2980,18 @@ void FStarshipEditorStyle::FStyle::SetupTutorialStyles()
 			.SetTypefaceFontName(TEXT("Italic")));
 
 		Set("Documentation.Header1", FTextBlockStyle(DocumentationHeaderText)
-			.SetFontSize(32));
+			.SetColorAndOpacity(FStyleColors::White)
+			.SetFontSize(16));
 
 		Set("Documentation.Header2", FTextBlockStyle(DocumentationHeaderText)
-			.SetFontSize(24));
+			.SetColorAndOpacity(FStyleColors::White)
+			.SetFontSize(14));
 
-		Set( "Documentation.Separator", new BOX_BRUSH( "Common/Separator", 1/4.0f, FLinearColor(1,1,1,0.5f) ) );
+		Set("Documentation.Header3", FTextBlockStyle(DocumentationHeaderText)
+			.SetColorAndOpacity(FStyleColors::White)
+			.SetFontSize(12));
+
+		Set( "Documentation.Separator", new FSlateColorBrush(FLinearColor::FromSRGBColor(FColor(59, 59, 59))));
 	}
 }
 
diff --git a/Engine/Source/Editor/UnrealEd/Private/AssetThumbnail.cpp b/Engine/Source/Editor/UnrealEd/Private/AssetThumbnail.cpp
--- a/Engine/Source/Editor/UnrealEd/Private/AssetThumbnail.cpp
+++ b/Engine/Source/Editor/UnrealEd/Private/AssetThumbnail.cpp
@@ -589,6 +589,7 @@ public:
 		const UClass* Class = FindObjectSafe<UClass>(AssetData.AssetClassPath);
 
 		TArray<FAssetDisplayInfo> OutSystemInfo;
+		TSharedRef<SBox> PromptBox = SNew(SBox);
 		TSharedRef<SVerticalBox> OverallTooltipVBox = SNew(SVerticalBox);
 		{
 			const FSlateBrush* ClassIcon = FAppStyle::GetDefaultBrush();
@@ -729,21 +730,13 @@ public:
 			// More info
 			if (!OutSystemInfo.IsEmpty())
 			{
-				OverallTooltipVBox->AddSlot()
-				.Padding(0.f, 2.f, 0.f, 0.f)
-				.AutoHeight()
-				.VAlign(VAlign_Center)
-				[
+				PromptBox->SetPadding(FMargin(0.f, 2.f, 0.f, 0.f));
+				PromptBox->SetContent(
 					SNew(SRichTextBlock)
-					.Visibility_Lambda([] ()
-					{
-						const FModifierKeysState ModifierKeys = FSlateApplication::Get().GetModifierKeys();
-						return ModifierKeys.IsAltDown() && ModifierKeys.IsControlDown() ? EVisibility::Collapsed : EVisibility::Visible;
-					})
 					.Justification(ETextJustify::Center)
 					.Text(NSLOCTEXT("AssetThumbnail", "MoreInfoTooltip", "Hold <WrappedCommand/> for more"))
-					+ SRichTextBlock::WidgetDecorator( TEXT("WrappedCommand"), this, &SAssetThumbnail::OnCreateWidgetDecoratorWidget )
-				];
+					+ SRichTextBlock::WidgetDecorator(TEXT("WrappedCommand"), this, &SAssetThumbnail::OnCreateWidgetDecoratorWidget)
+				);
 			}
 		}
 
@@ -769,6 +762,7 @@ public:
 
 		return SNew(SDocumentationToolTip)
 				.OverrideExtendedToolTipContent(ExtendedToolTip)
+				.OverridePromptContent(PromptBox)
 				[
 					OverallTooltipVBox
 				];
@@ -1578,7 +1572,7 @@ FAssetThumbnailPool::~FAssetThumbnailPool()
 {
 	UThumbnailManager* ThumbnailManager = UThumbnailManager::TryGet();
 	if ( IsValid( ThumbnailManager ) &&
-		 !ThumbnailManager->HasAnyFlags( RF_BeginDestroyed | RF_FinishDestroyed ) )
+		!ThumbnailManager->HasAnyFlags( RF_BeginDestroyed | RF_FinishDestroyed ) )
 	{
 		ThumbnailManager->GetOnThumbnailDirtied().RemoveAll(this);
 	}
diff --git a/Engine/Source/Runtime/SlateCore/Private/Styling/StarshipCoreStyle.cpp b/Engine/Source/Runtime/SlateCore/Private/Styling/StarshipCoreStyle.cpp
--- a/Engine/Source/Runtime/SlateCore/Private/Styling/StarshipCoreStyle.cpp
+++ b/Engine/Source/Runtime/SlateCore/Private/Styling/StarshipCoreStyle.cpp
@@ -295,7 +295,8 @@ TSharedRef<ISlateStyle> FStarshipCoreStyle::Create()
 		Style->Set("Icons.Denied", new IMAGE_BRUSH("Icons/denied_16x", Icon16x16));
 	
 		Style->Set("Icons.Help", new IMAGE_BRUSH_SVG("Starship/Common/help", Icon16x16));
-	
+		Style->Set("Icons.Help.Solid", new IMAGE_BRUSH_SVG("Starship/Common/help-solid", Icon16x16));
+
 		Style->Set("Icons.Download", new IMAGE_BRUSH("Icons/icon_Downloads_16x", Icon16x16));
 		Style->Set("Icons.Server", new IMAGE_BRUSH_SVG("Starship/Common/server", Icon16x16));
 		Style->Set("Icons.Cloud", new IMAGE_BRUSH("Icons/icon_Downloads_16x", Icon16x16));
