Enhanced tooltip: don't show the create/edit controls when the extended content is coming in through an override widget.

diff --git a/Engine/Source/Editor/Documentation/Private/SDocumentationToolTip.cpp b/Engine/Source/Editor/Documentation/Private/SDocumentationToolTip.cpp
--- a/Engine/Source/Editor/Documentation/Private/SDocumentationToolTip.cpp
+++ b/Engine/Source/Editor/Documentation/Private/SDocumentationToolTip.cpp
@@ -702,7 +702,7 @@ EVisibility SDocumentationToolTip::GetFullTipVisibility() const
 
 EVisibility SDocumentationToolTip::GetControlVisibility() const
 {
-	if (IsDisplayingDocumentationLink && (IsShowingFullTip || !DocumentationPage.IsValid() || !DocumentationPage->HasExcerpt(ExcerptName)))
+	if (IsDisplayingDocumentationLink && (IsShowingFullTip || !DocumentationPage.IsValid() || !DocumentationPage->HasExcerpt(ExcerptName)) && !OverrideFullTooltipContent.IsValid())
 	{
 		return EVisibility::Visible;
 	}
