Fix warning for missing slate assets
- Content directory wasn't set and parths were resolving relative to the executable. The "Selector" image wasn't being found.

diff --git a/Engine/Source/Editor/Documentation/Private/DocumentationStyleSet.cpp b/Engine/Source/Editor/Documentation/Private/DocumentationStyleSet.cpp
--- a/Engine/Source/Editor/Documentation/Private/DocumentationStyleSet.cpp
+++ b/Engine/Source/Editor/Documentation/Private/DocumentationStyleSet.cpp
@@ -43,6 +43,8 @@ FDocumentationStyleSet::FDocumentationStyleSet()
 	const FLinearColor ItemBackground = COLOR("#808080FF");              // Hover2
 	const FLinearColor ItemBackgroundHighlight = COLOR("#A0A0A0FF");     // New shade between Hover2 and Foreground
 
+	SetContentRoot(FPaths::EngineContentDir() / TEXT("Slate"));
+
 	Set("ToolTip.Header", new FSlateColorBrush(DarkHeaderBackground));
 
 	Set("ToolTip.ScrollBox", FScrollBoxStyle()
