Disable some engine init classes not needed for cooking and commandlets

#Synced-CL 51256439

diff --git a/Engine/Source/Editor/Documentation/Private/Documentation.cpp b/Engine/Source/Editor/Documentation/Private/Documentation.cpp
--- a/Engine/Source/Editor/Documentation/Private/Documentation.cpp
+++ b/Engine/Source/Editor/Documentation/Private/Documentation.cpp
@@ -32,6 +32,7 @@
 #include "UnrealEdMisc.h"
 #include "Widgets/DeclarativeSyntaxSupport.h"
 #include "Widgets/SToolTip.h"
+#include "Misc/App.h"
 
 class IDocumentationPage;
 class SVerticalBox;
@@ -48,26 +49,29 @@ TSharedRef< IDocumentation > FDocumentation::Create()
 
 FDocumentation::FDocumentation() 
 {
-	IMainFrameModule& MainFrameModule = FModuleManager::LoadModuleChecked<IMainFrameModule>(TEXT("MainFrame"));
-	MainFrameModule.OnMainFrameSDKNotInstalled().AddRaw(this, &FDocumentation::HandleSDKNotInstalled);
-
-	if (const UDocumentationSettings* DocSettings = GetDefault<UDocumentationSettings>())
+	if (FApp::CanEverRender() && !IsRunningCommandlet())
 	{
-		for (const FDocumentationBaseUrl& BaseUrl : DocSettings->DocumentationBaseUrls)
+		IMainFrameModule& MainFrameModule = FModuleManager::LoadModuleChecked<IMainFrameModule>(TEXT("MainFrame"));
+		MainFrameModule.OnMainFrameSDKNotInstalled().AddRaw(this, &FDocumentation::HandleSDKNotInstalled);
+
+		if (const UDocumentationSettings* DocSettings = GetDefault<UDocumentationSettings>())
 		{
-			if (!RegisterBaseUrl(BaseUrl.Id, BaseUrl.Url))
+			for (const FDocumentationBaseUrl& BaseUrl : DocSettings->DocumentationBaseUrls)
 			{
-				UE_LOG(LogDocumentation, Warning, TEXT("Could not register documentation base URL: %s"), *BaseUrl.Id);
+				if (!RegisterBaseUrl(BaseUrl.Id, BaseUrl.Url))
+				{
+					UE_LOG(LogDocumentation, Warning, TEXT("Could not register documentation base URL: %s"), *BaseUrl.Id);
+				}
 			}
 		}
-	}
 
-	AddSourcePath(FPaths::Combine(FPaths::ProjectDir(), TEXT("Documentation"), TEXT("Source")));
-	for (TSharedRef<IPlugin> Plugin : IPluginManager::Get().GetEnabledPlugins())
-	{
-		AddSourcePath(FPaths::Combine(Plugin->GetBaseDir(), TEXT("Documentation"), TEXT("Source")));
+		AddSourcePath(FPaths::Combine(FPaths::ProjectDir(), TEXT("Documentation"), TEXT("Source")));
+		for (TSharedRef<IPlugin> Plugin : IPluginManager::Get().GetEnabledPlugins())
+		{
+			AddSourcePath(FPaths::Combine(Plugin->GetBaseDir(), TEXT("Documentation"), TEXT("Source")));
+		}
+		AddSourcePath(FPaths::Combine(FPaths::EngineDir(), TEXT("Documentation"), TEXT("Source")));
 	}
-	AddSourcePath(FPaths::Combine(FPaths::EngineDir(), TEXT("Documentation"), TEXT("Source")));
 
 	RegisterConfigRedirects();
 }
